Add python based remote console client

#Synced-CL 52128428

diff --git a/Engine/Extras/UnrealRemoteConsole/remote_console/__init__.py b/Engine/Extras/UnrealRemoteConsole/remote_console/__init__.py
new file mode 100644
--- /dev/null
+++ b/Engine/Extras/UnrealRemoteConsole/remote_console/__init__.py
@@ -0,0 +1,32 @@
+"""
+remote_console -- Cross-platform client package for UE Remote Console Server.
+
+Platform-specific discovery backends live under
+Engine/Platforms/<Platform>/Extras/UnrealRemoteConsole/discovery.py
+and are auto-discovered at runtime (same convention as ushell).
+"""
+
+import sys
+
+if sys.platform == "win32":
+    def _setup_windows_console():
+        """Configure Windows console for UTF-8 output and ANSI color codes."""
+        import ctypes
+        ctypes.windll.kernel32.SetConsoleOutputCP(65001)
+        kernel32 = ctypes.windll.kernel32
+        for handle_id in (-11, -12):
+            handle = kernel32.GetStdHandle(handle_id)
+            mode = ctypes.c_ulong()
+            if kernel32.GetConsoleMode(handle, ctypes.byref(mode)):
+                kernel32.SetConsoleMode(handle, mode.value | 0x0004)
+        for stream in ("stdout", "stderr"):
+            s = getattr(sys, stream)
+            if hasattr(s, "reconfigure"):
+                s.reconfigure(encoding="utf-8", errors="replace")
+    _setup_windows_console()
+    del _setup_windows_console
+
+from .client import RemoteConsole, DEFAULT_TIMEOUT, DANGEROUS_COMMANDS
+from .interactive import interactive_loop, pipe_mode
+from .formatting import _format_log_msg, _disable_colors
+from .tests import _run_self_tests
diff --git a/Engine/Extras/UnrealRemoteConsole/remote_console/client.py b/Engine/Extras/UnrealRemoteConsole/remote_console/client.py
new file mode 100644
--- /dev/null
+++ b/Engine/Extras/UnrealRemoteConsole/remote_console/client.py
@@ -0,0 +1,500 @@
+"""
+RemoteConsole -- JSON-line protocol client for UE Remote Console Server.
+
+Single persistent TCP connection with bidirectional newline-delimited JSON.
+Log output streams from the server interleaved with command responses.
+Each request has an "id" field echoed in the response for async matching.
+"""
+
+import codecs
+import json
+import socket
+import sys
+import threading
+import time
+
+DEFAULT_TIMEOUT = 30.0
+_JSON_SEPS = (',', ':')
+
+DANGEROUS_COMMANDS = {"exit", "quit"}
+
+
+class RemoteConsole:
+    """JSON-line protocol remote console client.
+
+    Single persistent connection. Client sends JSON requests, server sends
+    JSON responses and unsolicited log events. All messages are newline-
+    delimited JSON.
+
+    Connection lifecycle:
+      1. TCP connect
+      2. Client sends:  {"id":N,"op":"hello","version":1}
+      3. Server responds: {"id":N,"op":"hello","version":1}
+      4. Server begins streaming log events
+      5. Bidirectional JSON-line exchange until disconnect
+
+    Client -> Server (after handshake):
+      {"id":1,"op":"exec","cmd":"stat unit"}
+      {"id":2,"op":"complete","q":"r.Nanite","offset":0,"limit":50}
+      {"id":3,"op":"getvar","n":"r.RHI.Name"}
+
+    Server -> Client:
+      {"id":1,"op":"exec","ok":true}
+      {"id":2,"op":"complete","results":[...],"offset":0,"count":N,"total":M}
+      {"id":3,"op":"getvar","name":"...","value":"...","source":"...","help":"..."}
+      {"op":"log","cat":"LogTemp","v":"Log","line":"Hello World"}
+    """
+
+    def __init__(self, host: str, port: int, timeout: float = DEFAULT_TIMEOUT, verbose: bool = False, log_file: str = None, decoded_log_file: str = None):
+        self.host = host
+        self.port = port
+        self.timeout = timeout
+        self.verbose = verbose
+        self.label = None          # display label for prompt (e.g., target name)
+        self._raw_log = None
+        self._decoded_log = None
+        try:
+            if log_file:
+                self._raw_log = open(log_file, "wb")
+            if decoded_log_file:
+                self._decoded_log = open(decoded_log_file, "w", encoding="utf-8")
+        except Exception:
+            if self._raw_log:
+                self._raw_log.close()
+            raise
+        self._sock = None
+        self._next_id = 1
+        self._pending = {}       # id -> threading.Event
+        self._responses = {}     # id -> parsed JSON dict
+        self._lock = threading.Lock()
+        self._reader_thread = None
+        self._stop_event = threading.Event()
+        self._reconnect_stop = threading.Event()
+        self._log_callback = None
+        # Accessed from multiple threads; CPython GIL ensures atomic
+        # read/write of simple attributes, which suffices here.
+        self._connected = False
+        self._server_version = 0
+        self._auto_reconnect = False
+        self._reconnect_thread = None
+        self._disconnect_callback = None
+        self._reconnect_callback = None
+        self._discover_endpoint = None
+
+    def _dbg(self, msg: str):
+        """Print a diagnostic message when verbose mode is enabled."""
+        if self.verbose:
+            print(f"  [debug] {msg}", file=sys.stderr)
+
+    @property
+    def connected(self) -> bool:
+        return self._connected
+
+    def connect(self, quiet: bool = False) -> bool:
+        """Connect to the server. Returns True on success."""
+        # Clean up any previous connection before establishing a new one.
+        self._cleanup_connection()
+        try:
+            sock = socket.socket(socket.AF_INET, socket.SOCK_STREAM)
+            try:
+                sock.settimeout(self.timeout)
+                sock.connect((self.host, self.port))
+            except (socket.error, OSError):
+                sock.close()
+                raise
+            self._sock = sock
+            self._connected = True
+            self._dbg(f"connect() TCP success, _connected={self._connected}")
+
+            # Start reader thread
+            self._stop_event.clear()
+            self._reader_thread = threading.Thread(
+                target=self._reader_loop, daemon=True, name="RemoteConsole-Reader")
+            self._reader_thread.start()
+
+            # Handshake: send {"op":"hello","version":1} and wait for the
+            # server's hello response.  This both validates that we're talking
+            # to a JSON-line remote console and lets the server defer log
+            # subscription until the client is ready.
+            try:
+                resp = self._send_request("hello", timeout=3.0, version=1)
+                if resp.get("error"):
+                    raise ConnectionError(resp["error"])
+                self._server_version = resp.get("version", 1)
+                self._dbg(f"handshake ok, server version={self._server_version}")
+            except Exception as e:
+                if not quiet:
+                    print(f"Handshake failed: {e}", file=sys.stderr)
+                self.close()
+                return False
+
+            # Even if the handshake round-tripped, the reader may have already
+            # died (e.g., server closed the connection right after responding).
+            if not self._connected:
+                if not quiet:
+                    print("Connection lost immediately after handshake.", file=sys.stderr)
+                self.close()
+                return False
+
+            return True
+        except (socket.error, OSError) as e:
+            if not quiet:
+                print(f"Connection failed: {e}", file=sys.stderr)
+            self._sock = None
+            self._connected = False
+            return False
+
+    def _send_request(self, op: str, timeout: float = None, **kwargs) -> dict:
+        """Send a JSON-line request and wait for the response."""
+        if not self._connected or self._sock is None:
+            self._dbg(f"_send_request rejected: _connected={self._connected}, _sock={'set' if self._sock else 'None'}")
+            raise ConnectionError("Not connected")
+        with self._lock:
+            req_id = self._next_id
+            self._next_id += 1
+            event = threading.Event()
+            self._pending[req_id] = event
+
+        msg = {"id": req_id, "op": op}
+        msg.update(kwargs)
+        line = json.dumps(msg, separators=_JSON_SEPS) + "\n"
+
+        sock = self._sock
+        if sock is None:
+            with self._lock:
+                self._pending.pop(req_id, None)
+            raise ConnectionError("Not connected")
+        try:
+            sock.sendall(line.encode("utf-8"))
+        except (socket.error, OSError) as e:
+            with self._lock:
+                self._pending.pop(req_id, None)
+            raise ConnectionError(f"Send failed: {e}")
+
+        # Wait for response, checking for disconnect periodically so we don't
+        # hang for the full timeout if the connection drops after sending.
+        effective_timeout = timeout if timeout is not None else self.timeout
+        deadline = time.monotonic() + effective_timeout
+        while not event.is_set():
+            remaining = deadline - time.monotonic()
+            if remaining <= 0:
+                with self._lock:
+                    self._pending.pop(req_id, None)
+                    resp = self._responses.pop(req_id, None)
+                if resp is not None:
+                    return resp
+                raise TimeoutError(f"No response for request {req_id} (op={op})")
+            # Wake every 0.5s to check if connection dropped
+            event.wait(timeout=min(remaining, 0.5))
+            if event.is_set():
+                break
+            if not self._connected:
+                with self._lock:
+                    self._pending.pop(req_id, None)
+                    resp = self._responses.pop(req_id, None)
+                if resp:
+                    return resp
+                raise ConnectionError("Connection lost while waiting for response")
+
+        with self._lock:
+            response = self._responses.pop(req_id, {})
+            self._pending.pop(req_id, None)
+
+        return response
+
+    def _reader_loop(self):
+        """Background thread: read JSON lines from the server, dispatch responses and log events."""
+        sock = self._sock
+        if not sock:
+            return
+
+        # Capture identity so we can detect if we're a stale reader from a previous connection.
+        my_sock = sock
+        sock.settimeout(0.5)
+        self._dbg(f"reader_loop started, stop_event={self._stop_event.is_set()}, sock={id(sock):#x}")
+
+        # Use an incremental decoder so multi-byte UTF-8 sequences split across
+        # recv() boundaries are reassembled correctly instead of producing U+FFFD.
+        # Genuine invalid bytes still produce U+FFFD which we strip below.
+        utf8_decoder = codecs.getincrementaldecoder("utf-8")("replace")
+        buf_parts = []   # accumulate decoded text chunks; join only when splitting lines
+        json_loads = json.loads  # local binding avoids global/module lookup per iteration
+
+        while not self._stop_event.is_set():
+            try:
+                data = sock.recv(65536)
+            except socket.timeout:
+                continue
+            except (socket.error, OSError) as e:
+                self._dbg(f"reader recv error: {e}")
+                break
+            if not data:
+                self._dbg("reader recv empty (EOF)")
+                break
+
+            if self._raw_log:
+                self._raw_log.write(data)
+                self._raw_log.flush()
+
+            decoded = utf8_decoder.decode(data)
+
+            # In verbose mode, report any replacement characters from genuine
+            # decode errors (not boundary splits -- the incremental decoder
+            # handles those). This helps diagnose server-side encoding issues.
+            if self.verbose and "\ufffd" in decoded:
+                count = decoded.count("\ufffd")
+                bad = []
+                i = 0
+                while i < len(data) and len(bad) < 5:
+                    try:
+                        data[i:].decode("utf-8")
+                        break
+                    except UnicodeDecodeError as ude:
+                        pos = i + ude.start
+                        end = i + ude.end
+                        ctx = data[max(0, pos - 3):min(len(data), end + 3)]
+                        bad.append(f"@{pos}: [{ctx.hex(' ')}]")
+                        i = i + ude.end
+                self._dbg(f"UTF-8 errors: {count} bad byte(s) in {len(data)}B chunk: {'; '.join(bad)}")
+
+            buf_parts.append(decoded)
+
+            # Only join and split when we have at least one complete line
+            if "\n" not in decoded:
+                continue
+
+            text_buf = "".join(buf_parts)
+            buf_parts.clear()
+
+            # Process complete lines
+            lines = text_buf.split("\n")
+            # Last element is the incomplete tail (possibly empty)
+            tail = lines[-1]
+            if tail:
+                buf_parts.append(tail)
+
+            for raw_line in lines[:-1]:
+                if not raw_line:
+                    continue
+                # Strip trailing \r (Windows line endings)
+                if raw_line[-1] == '\r':
+                    raw_line = raw_line[:-1]
+                    if not raw_line:
+                        continue
+
+                try:
+                    msg = json_loads(raw_line)
+                except json.JSONDecodeError:
+                    continue
+
+                op = msg.get("op", "")
+                msg_id = msg.get("id")
+
+                if op == "log":
+                    # Unsolicited log event
+                    if self._decoded_log:
+                        cat = msg.get("cat", "")
+                        text = msg.get("line", "")
+                        self._decoded_log.write(f"[{cat}] {text}\n")
+                        self._decoded_log.flush()
+                    cb = self._log_callback
+                    if cb:
+                        try:
+                            cb(msg)
+                        except Exception as _cb_err:
+                            self._dbg(f"log callback error: {_cb_err}")
+                elif msg_id is not None:
+                    # Response to a request -- only store if still pending
+                    # (avoids orphaned entries after caller already timed out)
+                    with self._lock:
+                        event = self._pending.get(msg_id)
+                        if event is not None:
+                            self._responses[msg_id] = msg
+                    if event:
+                        event.set()
+                else:
+                    # Unsolicited non-log message (e.g., parse error from server)
+                    cb = self._log_callback
+                    if cb:
+                        try:
+                            cb({"op": "log", "cat": "RemoteConsole",
+                                "v": "Error", "line": msg.get("msg", str(msg))})
+                        except Exception as _cb_err:
+                            self._dbg(f"log callback error: {_cb_err}")
+
+        # Flush any remaining bytes from the decoder (incomplete multi-byte tail).
+        utf8_decoder.decode(b"", final=True)
+
+        # Only update state if we're still the active reader (not a stale thread
+        # from a previous connection that outlived its join timeout).
+        if self._sock is not my_sock:
+            self._dbg("stale reader exiting")
+            return
+
+        self._connected = False
+        self._dbg("reader_loop setting _connected=False")
+        # Wake all pending requests so callers see the disconnect instead of
+        # hanging for the full timeout.  _cleanup_connection also does this,
+        # but the reader exits first, so we wake them here for promptness.
+        with self._lock:
+            for req_id, event in self._pending.items():
+                if req_id not in self._responses:
+                    self._responses[req_id] = {"error": "connection lost"}
+                event.set()
+        self._start_auto_reconnect()
+
+    def enable_auto_reconnect(self, on_disconnect=None, on_reconnect=None):
+        """Enable automatic reconnection when the connection drops.
+        Optional callbacks are invoked on disconnect/reconnect."""
+        self._auto_reconnect = True
+        self._disconnect_callback = on_disconnect
+        self._reconnect_callback = on_reconnect
+
+    def set_endpoint_discovery(self, callback):
+        """Set callback() -> (host, port, label) for endpoint rediscovery on reconnect.
+        Used when the server port may change between connections."""
+        self._discover_endpoint = callback
+
+    def try_rediscover_endpoint(self) -> bool:
+        """Re-discover the endpoint if a discovery callback is set.
+        Updates host/port/label on success. Returns True if updated."""
+        if not self._discover_endpoint:
+            return False
+        result = self._discover_endpoint()
+        if not result or len(result) < 2 or result[0] is None:
+            return False
+        self.host, self.port = result[0], result[1]
+        if len(result) > 2 and result[2]:
+            self.label = result[2]
+        return True
+
+    def _start_auto_reconnect(self):
+        """Start the auto-reconnect loop if enabled."""
+        if not self._auto_reconnect:
+            return
+        if self._reconnect_thread and self._reconnect_thread.is_alive():
+            return
+        cb = self._disconnect_callback
+        if cb:
+            try:
+                cb()
+            except Exception as _cb_err:
+                self._dbg(f"disconnect callback error: {_cb_err}")
+        self._reconnect_stop.clear()
+        self._reconnect_thread = threading.Thread(
+            target=self._reconnect_loop, daemon=True, name="RemoteConsole-Reconnect")
+        self._reconnect_thread.start()
+
+    def _reconnect_loop(self):
+        """Periodically attempt to reconnect until successful or stopped.
+        Uses _reconnect_stop (not _stop_event) so that connect()'s internal
+        _cleanup_connection calls don't interfere with the reconnect sleep."""
+        while not self._reconnect_stop.is_set():
+            self._reconnect_stop.wait(timeout=2.0)
+            if self._reconnect_stop.is_set() or not self._auto_reconnect:
+                break
+            if self._connected:
+                break
+            # Re-discover endpoint if available (port may change)
+            if self._discover_endpoint:
+                if not self.try_rediscover_endpoint():
+                    continue  # endpoint not available yet, retry later
+            if self.connect(quiet=True):
+                # Verify the connection is still alive after connect returned.
+                # The reader thread may have already detected EOF (e.g., server
+                # accepted on OS backlog but isn't actually serving yet).
+                if not self._connected:
+                    continue
+                self._dbg(f"reconnect succeeded, _connected={self._connected}")
+                cb = self._reconnect_callback
+                if cb:
+                    try:
+                        cb()
+                    except Exception as _cb_err:
+                        self._dbg(f"reconnect callback error: {_cb_err}")
+                # If connection dropped during/after the callback, keep retrying
+                # instead of exiting (fixes death race where auto-reconnect dies
+                # silently on immediate post-reconnect disconnect).
+                if not self._connected:
+                    continue
+                break
+
+    def execute(self, command: str) -> dict:
+        """Execute a console command. Returns the response dict."""
+        return self._send_request("exec", cmd=command)
+
+    def complete(self, prefix: str, offset: int = 0, limit: int = 100) -> tuple:
+        """Tab completion. Returns (results, total).
+        Results is a list of dicts with keys: name, value, source, help."""
+        resp = self._send_request("complete", q=prefix, offset=offset, limit=limit)
+        results = resp.get("results", [])
+        total = resp.get("total", 0)
+        # Normalize results to tuples for backward compat with completer
+        tuples = []
+        for r in results:
+            tuples.append((
+                r.get("name", ""),
+                r.get("value", ""),
+                r.get("source", ""),
+                r.get("help", "")
+            ))
+        return tuples, total
+
+    def getvar(self, name: str) -> dict:
+        """Get a CVar's current value. Returns the response dict."""
+        return self._send_request("getvar", n=name)
+
+    def _cleanup_connection(self):
+        """Tear down socket and reader thread without affecting auto-reconnect state."""
+        self._stop_event.set()
+        if self._sock:
+            try:
+                self._sock.close()
+            except Exception:
+                pass
+            self._sock = None
+        if self._reader_thread and self._reader_thread.is_alive():
+            self._reader_thread.join(timeout=2.0)
+        self._reader_thread = None
+        self._connected = False
+        with self._lock:
+            for req_id, event in self._pending.items():
+                self._responses[req_id] = {"error": "connection closed"}
+                event.set()
+            self._pending.clear()
+            self._responses.clear()
+
+    def close(self, permanent=False):
+        """Close the connection. If permanent=True, also disables auto-reconnect."""
+        if permanent:
+            self._auto_reconnect = False
+            self._reconnect_stop.set()
+            if self._raw_log:
+                self._raw_log.close()
+                self._raw_log = None
+            if self._decoded_log:
+                self._decoded_log.close()
+                self._decoded_log = None
+        self._cleanup_connection()
+        if permanent and self._reconnect_thread and self._reconnect_thread.is_alive():
+            self._reconnect_thread.join(timeout=3.0)
+
+    def reconnect(self):
+        """Manually reconnect. Stops any running auto-reconnect, tears down the
+        current connection, attempts one connect, and restarts auto-reconnect
+        on failure if it was previously enabled."""
+        saved_auto = self._auto_reconnect
+        self._auto_reconnect = False
+        self._reconnect_stop.set()
+        self._cleanup_connection()
+        if self._reconnect_thread and self._reconnect_thread.is_alive():
+            self._reconnect_thread.join(timeout=3.0)
+        self._reconnect_stop.clear()
+        self._auto_reconnect = saved_auto
+        if not self.connect(quiet=True):
+            self._start_auto_reconnect()
+
+    def set_log_callback(self, callback):
+        """Set callback(msg: dict) for async log output.
+        msg has keys: op, cat, v, line."""
+        self._log_callback = callback
diff --git a/Engine/Extras/UnrealRemoteConsole/remote_console/commands.py b/Engine/Extras/UnrealRemoteConsole/remote_console/commands.py
new file mode 100644
--- /dev/null
+++ b/Engine/Extras/UnrealRemoteConsole/remote_console/commands.py
@@ -0,0 +1,165 @@
+"""Internal slash commands and handler for the UE Remote Console REPL."""
+
+from __future__ import annotations
+
+import os
+import time
+from typing import TYPE_CHECKING
+
+if TYPE_CHECKING:
+    from .client import RemoteConsole
+
+from .formatting import _format_log_msg
+
+
+INTERNAL_COMMANDS_COMMON = {
+    "/help":      "Show this help",
+    "/h":         "Show this help",
+    "/?":         "Show this help",
+    "/quit":      "Disconnect and exit",
+    "/exit":      "Disconnect and exit",
+    "/q":         "Disconnect and exit",
+    "/reconnect": "Reconnect to the target",
+    "/clear":     "Clear the console screen",
+    "/silent":    "Execute a command silently (full processing, no output). Usage: /silent <cmd>",
+}
+
+INTERNAL_COMMANDS_BASIC_ONLY = {
+    "/complete":  "Show matching commands with help text.  Usage: /complete <prefix>",
+    "/get":       "Get a CVar's current value.  Usage: /get <cvar>",
+}
+
+INTERNAL_COMMANDS = {**INTERNAL_COMMANDS_COMMON, **INTERNAL_COMMANDS_BASIC_ONLY}
+
+INTERNAL_HELP_TEXT_COMMON = """  Internal commands (prefix with /):
+    /help, /h, /?          Show this help
+    /quit, /exit, /q       Disconnect and exit
+    /reconnect             Reconnect to the target
+    /clear                 Clear the console screen
+    /silent <cmd>          Execute command silently (full processing, no output)
+
+  Everything else is sent as a UE console command."""
+
+INTERNAL_HELP_TEXT_BASIC = """  Internal commands (prefix with /):
+    /help, /h, /?          Show this help
+    /quit, /exit, /q       Disconnect and exit
+    /complete <prefix>     Show matching commands with help text
+    /get <cvar>            Get a CVar's current value
+    /reconnect             Reconnect to the target
+    /clear                 Clear the console screen
+    /silent <cmd>          Execute command silently (full processing, no output)
+
+  Everything else is sent as a UE console command."""
+
+
+
+
+def _handle_internal_command(cmd: str, console: RemoteConsole, log_fn=print, rich_mode=False) -> bool:
+    """Handle a /command. Returns True to continue the loop, False to exit."""
+    internal = cmd[1:].strip()
+    args_str = ""
+    if " " in internal:
+        internal, args_str = internal.split(" ", 1)
+    internal = internal.lower()
+
+    if internal in ("quit", "exit", "q"):
+        return False
+    elif internal in ("help", "h", "?"):
+        log_fn(INTERNAL_HELP_TEXT_COMMON if rich_mode else INTERNAL_HELP_TEXT_BASIC)
+    elif internal == "complete":
+        if not args_str:
+            log_fn("  Usage: /complete <prefix>")
+        else:
+            results, total = console.complete(args_str)
+            if results:
+                max_name = max(len(name) for name, _, _, _ in results)
+                max_name = max(max_name, 20)
+                for name, value, source, help_text in results:
+                    parts = []
+                    if value:
+                        parts.append(f"={value}")
+                    if source:
+                        parts.append(f"({source})")
+                    if help_text:
+                        parts.append(help_text)
+                    meta = " ".join(parts)
+                    if meta:
+                        log_fn(f"  {name:<{max_name}}  {meta}")
+                    else:
+                        log_fn(f"  {name}")
+                if total > len(results):
+                    log_fn(f"  ... showing {len(results)} of {total} matches")
+                else:
+                    log_fn(f"  ({total} matches)")
+            else:
+                log_fn("  (no matches)")
+    elif internal == "get":
+        if not args_str:
+            log_fn("  Usage: /get <cvar_name>")
+        else:
+            resp = console.getvar(args_str)
+            if "error" in resp:
+                log_fn(f"  {resp['error']}")
+            else:
+                parts = [f"{resp.get('name', args_str)} = \"{resp.get('value', '')}\""]
+                source = resp.get("source", "")
+                if source:
+                    parts.append(f"(SetBy: {source})")
+                help_text = resp.get("help", "")
+                if help_text:
+                    parts.append(f"Help: {help_text}")
+                log_fn(f"  {' '.join(parts)}")
+    elif internal == "clear":
+        if rich_mode:
+            from prompt_toolkit.shortcuts import clear
+            clear()
+        else:
+            os.system("cls" if os.name == "nt" else "clear")
+    elif internal == "reconnect":
+        log_fn("  Reconnecting...")
+        console.reconnect()
+        if console.connected:
+            log_fn("  Reconnected.")
+        else:
+            log_fn("  Reconnect failed. Auto-reconnect will keep trying.")
+    elif internal == "silent":
+        if not args_str:
+            log_fn("  Usage: /silent <command>")
+        else:
+            # Execute the command through the full pipeline (including log
+            # formatting and syntax highlighting on the callback thread),
+            # but suppress the final print.  This is useful for benchmarking
+            # the formatting path or executing commands whose output is not
+            # needed interactively.
+            #
+            # Because output arrives asynchronously via the log channel, we
+            # keep the silent callback active until 100ms of silence (no new
+            # log messages), which indicates the burst has finished.
+            last_msg_time = time.monotonic()
+            msg_count = 0
+            saved_callback = console._log_callback
+            def _silent_callback(msg):
+                nonlocal last_msg_time, msg_count
+                # Run the full formatting pipeline, discard result.
+                _format_log_msg(msg)
+                last_msg_time = time.monotonic()
+                msg_count += 1
+            console._log_callback = _silent_callback
+            try:
+                start = time.monotonic()
+                resp = console.execute(args_str)
+                if resp.get("error"):
+                    log_fn(f"  Error: {resp['error']}")
+                # Wait until 100ms of silence on the log stream.
+                last_msg_time = time.monotonic()
+                while True:
+                    time.sleep(0.02)
+                    if time.monotonic() - last_msg_time >= 0.1:
+                        break
+                elapsed = time.monotonic() - start
+                log_fn(f"  (silent: {msg_count} log lines processed in {elapsed:.3f}s)")
+            finally:
+                console._log_callback = saved_callback
+    else:
+        log_fn(f"  Unknown command: /{internal}  (type /help for list)")
+    return True
diff --git a/Engine/Extras/UnrealRemoteConsole/remote_console/formatting.py b/Engine/Extras/UnrealRemoteConsole/remote_console/formatting.py
new file mode 100644
--- /dev/null
+++ b/Engine/Extras/UnrealRemoteConsole/remote_console/formatting.py
@@ -0,0 +1,418 @@
+"""Log output formatting: syntax highlighting, color-coding, category colors."""
+
+import re
+import zlib
+
+_COLOR_ENABLED = True
+
+# UE log messages are color-coded to improve readability. The rules are
+# listed below. Each rule references the --test assertions that verify it.
+#
+# 1. LOG CATEGORIES -- e.g., [LogRHI], [LogEngine], [None]
+#    - The server sends the category in the "cat" field; it is rendered as
+#      a leading [Category] tag.
+#      Tests: "category [LogEngine] colored", "category preserved in plain text"
+#    - Additional leading [Tag] patterns in the line text (e.g., [SubSystem])
+#      are also detected and colored.
+#      Tests: "leading [SubSystem] tag colored", "leading tag text preserved"
+#    - Each category gets a deterministic color via CRC32 hash into a curated
+#      256-color palette. The color is stable across runs.
+#      Test: "category color is deterministic"
+#
+# 2. VERBOSITY -- Fatal, Error, Warning
+#    - Fatal:   bold bright red, with "Fatal: " prefix
+#      Test: "fatal verbosity label shown"
+#    - Error:   bright red, with "Error: " prefix
+#      Tests: "error verbosity label shown", "error text preserved"
+#    - Warning: bright yellow, with "Warning: " prefix
+#      Test: "warning verbosity label shown"
+#    - Log/Display/other: no special styling, no prefix
+#      Test: "no verbosity label for Log level"
+#    - Verbosity color applies to the entire message text (after categories).
+#
+# 3. QUOTED STRINGS -- "double" and 'single' quoted
+#    - Quote marks in dim yellow, string contents in bright yellow.
+#      Tests: "double quote mark dim", "double quote content bright",
+#             "single quote mark dim", "single quote content bright"
+#    - The brightness difference visually separates delimiters from content.
+#      Test: "quoted string '\"1\"' colored"
+#    - Numbers inside quoted strings are still highlighted using number
+#      coloring rules (bright cyan for value, dim cyan for unit suffix).
+#      Tests: "number inside double-quoted string colored",
+#             "quote marks still dim around number",
+#             "number inside single-quoted string colored",
+#             "unit suffix inside quoted string colored"
+#    - No text is eaten by highlighting.
+#      Test: "no text eaten by string highlight"
+#
+# 4. KEY-VALUE PAIRS -- key:value and key=value
+#    - Key name in blue, separator (: or =) in bright magenta (same as ->).
+#      Tests: "key 'rhi.syncinterval' colored", "separator ':' colored",
+#             "key 'Width' colored in key=value"
+#    - Key must start with a letter/underscore, followed by alphanumerics/dots.
+#    - Only matches when separator is followed by a non-space value character.
+#      A space between ":" and value is allowed only when the value starts
+#      with a digit (e.g., "Textures: 0" matches, "LastSetBy: Constructor" does not).
+#      Spaces around "=" are allowed (e.g., 'r.Foo = "40"' matches, as do
+#      'r.Foo ="40"' and 'r.Foo= "40"').
+#      Tests: "spaced = key colored", "spaced = separator colored",
+#             "spaced = text preserved", "key= value key colored",
+#             "key =value key colored"
+#    - Does NOT match :: or == (C++ scope resolution / comparison).
+#      Tests: ":: not treated as key:value", ":: text preserved"
+#    - Matches "Key: value" with space after separator.
+#      Tests: "'Textures' colored as key with space",
+#             "'Buffers' colored as key with space"
+#    - Examples: r.Nanite:1, Width=1920, rhi.syncinterval:0, Textures: 0,
+#      r.DynamicRes.TestScreenPercentage = "40"
+#    - No text is eaten by highlighting.
+#      Tests: "no text eaten by key:value", "number '1920' colored in key=value"
+#    - Standalone = (not part of key=value) is also highlighted.
+#      Test: "standalone = highlighted"
+#
+# 5. NUMBERS -- standalone numeric literals
+#    - Bright cyan. Integers, floats, scientific notation.
+#      Tests: "number '1' colored", "standalone number highlighted",
+#             "512 number colored", "2.4 number colored", "3.5 number colored",
+#             "60 number colored", "99.1 number colored"
+#    - NOT highlighted when part of an identifier (e.g., SpotLight77).
+#      Tests: "number in identifier not highlighted",
+#             "identifier+number text preserved"
+#    - Unit suffixes are rendered in a dimmer cyan to distinguish them
+#      from the numeric value:
+#      - Percent: 50%, 0.1%
+#        Test: "% suffix colored"
+#      - Size units: KB, KiB, MB, MiB, GB, GiB, K, M, G (case-mixed)
+#        Tests: "KB suffix colored differently", "space+MB suffix colored",
+#               "space+GiB suffix colored"
+#      - Time units: s, ms, us, ns
+#        Tests: "ms suffix colored", "ns suffix colored", "us suffix colored",
+#               "s suffix colored"
+#      - Pixels: px
+#        Test: "px suffix colored"
+#      - Frequency: Hz, hz
+#        Test: "Hz suffix colored"
+#      - Rate suffixes: MB/s, Kb/s, etc.
+#        Test: "MB/s suffix colored"
+#    - Both attached (512KB) and detached (512 KB) forms are matched.
+#      In both cases the number is bright cyan and the suffix is dim cyan.
+#      Tests: "2.4 with space+MB colored", "1 with space+GiB colored"
+#    - No text is eaten by highlighting.
+#      Test: "unit text preserved"
+#
+# 5b. DIMENSIONS -- NUMxNUM, NUMxNUMxNUM
+#    - Numbers in bright cyan, 'x' separator in bright magenta (same as ->).
+#      Tests: "dimension number colored", "dimension x uses arrow color",
+#             "dimension second number colored", "dimension text preserved"
+#    - Examples: 1920x1080, 256x256x4, 1067x600
+#
+# 6. ARROWS, OPERATORS, and SEPARATORS -- ->, =, x (in dimensions)
+#    - Bright magenta. Used in CVar change notifications (e.g., 1 -> 0),
+#      standalone assignment, and as the 'x' separator in dimensions.
+#      Tests: "arrow '->' colored", "arrow -> highlighted",
+#             "standalone = highlighted", "dimension x uses arrow color"
+#
+# 7. BRACKET CHARACTERS -- [ ] ( ) { } < >
+#    - Medium gray (subtly brighter than surrounding text).
+#      Tests: "opening paren highlighted", "closing paren highlighted",
+#             "angle bracket < highlighted", "angle bracket > highlighted"
+#    - Adjacent brackets are merged into one span (e.g., [[ or ]]).
+#      Tests: "double brackets [[ highlighted", "double brackets ]] highlighted"
+#    - Only the bracket characters themselves are colored, not their contents.
+#    - No text is eaten by highlighting.
+#      Test: "no text eaten by brackets"
+#
+# All color-coding can be disabled with --no-color.
+#
+# Important: No text is ever removed or reordered. All highlighting is purely
+# additive (ANSI escape codes inserted around tokens). The plain-text content
+# is preserved exactly as received from the server.
+# ---------------------------------------------------------------------------
+
+_CATEGORY_COLOR_PALETTE = [
+    31, 32, 33, 34, 35, 36,       # standard bright colors (skip black/white)
+    91, 92, 93, 94, 95, 96,       # high-intensity variants
+    130, 136, 142, 148, 154, 160, # 256-color range: warm
+    166, 172, 178, 184, 37, 43,   # 256-color range: mid
+    49, 50, 51, 75, 81, 87,      # 256-color range: cool
+    117, 123, 129, 135, 141, 147, # 256-color range: purple/blue
+    171, 177, 183, 189, 209, 215, # 256-color range: pink/salmon
+]
+
+_VERBOSITY_STYLES = {
+    "Fatal":   "\033[1;91m",    # bold bright red
+    "Error":   "\033[91m",      # bright red
+    "Warning": "\033[93m",      # bright yellow
+}
+
+_COLOR_ENABLED = True          # set to False by --no-color
+
+_RESET = "\033[0m"
+_STR_COLOR = "\033[93m"        # bright yellow for quoted string contents
+_QUOTE_COLOR = "\033[33m"     # yellow for quote marks -- dimmer than contents
+_NUM_COLOR = "\033[96m"        # bright cyan for numbers
+_KEY_COLOR = "\033[34m"        # blue for key names in key:value / key=value
+_UNIT_COLOR = "\033[36m"       # cyan for unit suffixes (KB, ms, Hz, etc.) -- dimmer than numbers
+_BRACKET_COLOR = "\033[38;5;245m"  # medium gray for brackets/parens/angles
+_ARROW_COLOR = "\033[95m"          # bright magenta for -> arrows
+
+
+def _disable_colors():
+    """Disable all ANSI color output."""
+    global _COLOR_ENABLED, _RESET, _STR_COLOR, _QUOTE_COLOR, _NUM_COLOR, _KEY_COLOR
+    global _UNIT_COLOR, _BRACKET_COLOR, _ARROW_COLOR
+    _COLOR_ENABLED = False
+    _RESET = ""
+    _STR_COLOR = ""
+    _QUOTE_COLOR = ""
+    _NUM_COLOR = ""
+    _KEY_COLOR = ""
+    _UNIT_COLOR = ""
+    _BRACKET_COLOR = ""
+    _ARROW_COLOR = ""
+
+# Matches tokens for syntax highlighting in log text.
+# Order matters: strings first, then key=value/key:value, then numbers,
+# then arrows, then single bracket characters.
+_HIGHLIGHT_RE = re.compile(
+    r"""(?:"(?:[^"\\]|\\.)*")"""       # "double quoted"
+    r"""|(?:'(?:[^'\\]|\\.)*')"""      # 'single quoted'
+    r"""|([A-Za-z_][A-Za-z0-9_.]*)"""  # key name (group 1) ...
+    r"""( ?= ?|: (?=[-+\d.])|:)"""       # ... separator (group 2): = with optional spaces, ": " before number, or bare :
+    r"""(?![:=])"""                     # ... but not :: or == (scope/comparison)
+    r"""(?=\S)"""                       # ... followed by a non-space value
+    r"""|(?<![a-zA-Z_0-9.])"""         # standalone numbers: not preceded by ident
+    r"""[-+]?(?:\d+\.?\d*|\.\d+)"""
+    r"""(?:[eE][-+]?\d+)?"""
+    r"""( ?"""                          # unit suffix (group 3), optional space + unit:
+    r"""(?:"""
+    r"""[KMGT]i?[Bb]"""               #   KiB, MB, GB, KB, etc.
+    r"""|[kKMGT][Bb]?"""              #   K, M, G, k, KB, etc.
+    r"""|[mun]?s"""                    #   s, ms, us, ns
+    r"""|px"""                         #   px (pixels)
+    r"""|[Hh]z"""                      #   Hz, hz
+    r"""|%"""                          #   percent
+    r"""|bytes?"""                     #   byte, bytes
+    r""")"""
+    r"""(?:/s)?"""                     #   optional /s rate suffix (MB/s, Kb/s, etc.)
+    r""")?"""
+    r"""(?![a-zA-Z_0-9.])"""          # not followed by identifier char
+    r"""|(?<![a-zA-Z_0-9.])"""         # dimensions: NUMxNUM or NUMxNUMxNUM
+    r"""(\d+(?:\.\d+)?(?:x\d+(?:\.\d+)?)+)"""  # group 4: full dimension string
+    r"""(?![a-zA-Z_0-9.])"""
+    r"""|->|="""                        # arrow and assignment operators
+    r"""|[\[\](){}<>]"""               # bracket characters
+)
+
+
+_BRACKET_CHARS = frozenset('[](){}<>')
+_DIM_SPLIT_RE = re.compile(r'(x)')      # for splitting dimension strings like "1920x1080"
+
+
+# Number pattern for highlighting inside quoted strings.
+# Same rules as the main regex but without key-value, arrows, brackets, or strings.
+_INNER_NUM_RE = re.compile(
+    r"""(?<![a-zA-Z_0-9.])"""
+    r"""(?:(\d+(?:\.\d+)?(?:x\d+(?:\.\d+)?)+)"""  # group 1: dimension (NUMxNUM)
+    r"""|[-+]?(?:\d+\.?\d*|\.\d+)"""
+    r"""(?:[eE][-+]?\d+)?"""
+    r"""( ?(?:[KMGT]i?[Bb]|[kKMGT][Bb]?|[mun]?s|px|[Hh]z|%|bytes?)(?:/s)?)?)"""  # group 2: unit
+    r"""(?![a-zA-Z_0-9.])"""
+)
+
+
+def _highlight_inner_string(text: str, base_style: str = "") -> str:
+    """Highlight numbers and dimensions inside a quoted string."""
+    result = []
+    last_end = 0
+    for m in _INNER_NUM_RE.finditer(text):
+        start, end = m.start(), m.end()
+        token = m.group()
+        if start > last_end:
+            result.append(text[last_end:start])
+        if m.group(1) is not None:
+            # Dimension pattern: split on 'x', color numbers bright, x as arrow
+            parts = _DIM_SPLIT_RE.split(token)
+            for part in parts:
+                if part == 'x':
+                    result.append(f"{_ARROW_COLOR}{part}{_RESET}{base_style}")
+                else:
+                    result.append(f"{_NUM_COLOR}{part}{_RESET}{base_style}")
+        else:
+            unit = m.group(2)
+            if unit:
+                num_part = token[:-len(unit)]
+                result.append(f"{_NUM_COLOR}{num_part}{_RESET}{base_style}"
+                              f"{_UNIT_COLOR}{unit}{_RESET}{base_style}")
+            else:
+                result.append(f"{_NUM_COLOR}{token}{_RESET}{base_style}")
+        last_end = end
+    if last_end < len(text):
+        result.append(text[last_end:])
+    return "".join(result) if result else text
+
+
+def _highlight_log_text(text: str, base_style: str = "") -> str:
+    """Apply syntax highlighting to a log message line.
+
+    Colorizes quoted strings, numeric literals, and bracket characters.
+    Consecutive bracket characters (e.g., '[[' or ']]') are merged into
+    a single colored span.  base_style is re-applied after each highlighted
+    token so that verbosity coloring (e.g., red for errors) is preserved.
+    """
+    result = []
+    last_end = 0
+    pending_brackets = []
+
+    def _flush_brackets():
+        if pending_brackets:
+            result.append(f"{_BRACKET_COLOR}{''.join(pending_brackets)}{_RESET}{base_style}")
+            pending_brackets.clear()
+
+    for m in _HIGHLIGHT_RE.finditer(text):
+        start, end = m.start(), m.end()
+        token = m.group()
+
+        if token in _BRACKET_CHARS:
+            # If there's a gap between last match and this bracket, flush
+            # any pending brackets and append the gap text first.
+            if start > last_end:
+                _flush_brackets()
+                result.append(text[last_end:start])
+            pending_brackets.append(token)
+            last_end = end
+            continue
+
+        # Non-bracket token -- flush any pending brackets first
+        _flush_brackets()
+
+        # Append any unhighlighted text before this match
+        if start > last_end:
+            result.append(text[last_end:start])
+
+        if token[0] in ('"', "'"):
+            # Quote marks dimmer, contents brighter, with numbers highlighted
+            quote_char = token[0]
+            inner = token[1:-1]
+            # Highlight numbers inside the string content
+            highlighted_inner = _highlight_inner_string(inner, base_style=_STR_COLOR)
+            result.append(f"{_QUOTE_COLOR}{quote_char}{_RESET}{base_style}"
+                          f"{_STR_COLOR}{highlighted_inner}{_RESET}{base_style}"
+                          f"{_QUOTE_COLOR}{quote_char}{_RESET}{base_style}")
+        elif token in ('->', '='):
+            result.append(f"{_ARROW_COLOR}{token}{_RESET}{base_style}")
+        elif m.group(1) is not None:
+            # Key-value pattern: color key and separator separately
+            key = m.group(1)
+            sep = m.group(2)
+            result.append(f"{_KEY_COLOR}{key}{_RESET}{base_style}{_ARROW_COLOR}{sep}{_RESET}{base_style}")
+        elif m.group(4) is not None:
+            # Dimension pattern (e.g., 1920x1080): numbers bright, 'x' dim
+            parts = _DIM_SPLIT_RE.split(token)
+            for part in parts:
+                if part == 'x':
+                    result.append(f"{_ARROW_COLOR}{part}{_RESET}{base_style}")
+                else:
+                    result.append(f"{_NUM_COLOR}{part}{_RESET}{base_style}")
+        else:
+            # Number token -- render unit suffix in a different color if present
+            unit = m.group(3)
+            if unit:
+                num_part = token[:-len(unit)]
+                result.append(f"{_NUM_COLOR}{num_part}{_RESET}{base_style}"
+                              f"{_UNIT_COLOR}{unit}{_RESET}{base_style}")
+            else:
+                result.append(f"{_NUM_COLOR}{token}{_RESET}{base_style}")
+
+        last_end = end
+
+    _flush_brackets()
+
+    # Append any remaining text after the last match
+    if last_end < len(text):
+        result.append(text[last_end:])
+
+    return "".join(result) if result else text
+
+
+_category_color_cache = {}
+
+
+def _category_color(cat: str) -> str:
+    """Return an ANSI escape for a deterministic color based on category name."""
+    if not _COLOR_ENABLED:
+        return ""
+    cached = _category_color_cache.get(cat)
+    if cached is not None:
+        return cached
+    # Use CRC32 instead of hash() which is randomized per process.
+    h = zlib.crc32(cat.encode("utf-8")) & 0xFFFFFFFF
+    idx = h % len(_CATEGORY_COLOR_PALETTE)
+    code = _CATEGORY_COLOR_PALETTE[idx]
+    # Always use 256-color foreground format to avoid accidentally setting
+    # background colors (codes 40-47) or other non-foreground attributes.
+    result = f"\033[38;5;{code}m"
+    _category_color_cache[cat] = result
+    return result
+
+
+# Matches leading [Tag] sequences that look like UE log categories.
+# Must start with a letter, contain only alphanumeric/underscore/colon chars.
+_LEADING_TAG_RE = re.compile(r'\[([A-Za-z][A-Za-z0-9_:]*)\]')
+
+# Verbosity levels that don't get a visible label prefix
+_QUIET_VERBOSITIES = frozenset({"Log", "Display", "Verbose", "VeryVerbose", ""})
+
+
+def _split_leading_tags(text):
+    """Extract and colorize leading [Tag] patterns, return (colored_prefix, remainder)."""
+    pos = 0
+    tag_parts = []
+    while pos < len(text):
+        m = _LEADING_TAG_RE.match(text, pos)
+        if m:
+            tag_name = m.group(1)
+            tag_parts.append(f"{_category_color(tag_name)}[{tag_name}]{_RESET}")
+            pos = m.end()
+            # Preserve whitespace between tags
+            while pos < len(text) and text[pos] == ' ':
+                tag_parts.append(' ')
+                pos += 1
+        else:
+            break
+    return "".join(tag_parts), text[pos:]
+
+
+def _format_log_msg(msg: dict) -> str:
+    """Format a log event dict as a colored display string."""
+    cat = msg.get("cat", "")
+    verbosity = msg.get("v", "")
+    line = msg.get("line", "")
+
+    # Fast path: no color processing needed
+    if not _COLOR_ENABLED:
+        prefix = f"[{cat}] " if cat else ""
+        v_label = f"{verbosity}: " if verbosity and verbosity not in _QUIET_VERBOSITIES else ""
+        return f"{prefix}{v_label}{line}"
+
+    # Apply verbosity styling to the whole line for errors/warnings/fatals
+    v_style = _VERBOSITY_STYLES.get(verbosity, "")
+    v_reset = _RESET if v_style else ""
+
+    # Color the server-provided category
+    prefix = ""
+    if cat:
+        prefix = f"{_category_color(cat)}[{cat}]{_RESET} "
+
+    # Show verbosity label for non-default levels so errors/warnings are obvious
+    verbosity_prefix = ""
+    if verbosity and verbosity not in _QUIET_VERBOSITIES:
+        verbosity_prefix = f"{v_style}{verbosity}: {_RESET}"
+
+    tag_prefix, remainder = _split_leading_tags(line)
+
+    # Syntax-highlight only the non-tag remainder
+    highlighted = _highlight_log_text(remainder, base_style=v_style)
+
+    full = f"{prefix}{tag_prefix}{verbosity_prefix}{v_style}{highlighted}{v_reset}"
+    return full
diff --git a/Engine/Extras/UnrealRemoteConsole/remote_console/interactive.py b/Engine/Extras/UnrealRemoteConsole/remote_console/interactive.py
new file mode 100644
--- /dev/null
+++ b/Engine/Extras/UnrealRemoteConsole/remote_console/interactive.py
@@ -0,0 +1,477 @@
+"""Interactive REPL loops for the UE Remote Console.
+
+Rich mode (prompt_toolkit): tab-completion, ghost text, history, batched log output.
+Basic mode (input()): simple fallback when prompt_toolkit is not installed.
+"""
+
+import os
+import sys
+import threading
+import time
+
+import socket
+
+from .client import RemoteConsole, DANGEROUS_COMMANDS
+from .commands import _handle_internal_command, INTERNAL_COMMANDS_COMMON
+from .formatting import _format_log_msg
+
+
+# ---------------------------------------------------------------------------
+
+try:
+    from prompt_toolkit import PromptSession
+    from prompt_toolkit.completion import Completer, Completion, ThreadedCompleter
+    from prompt_toolkit.history import FileHistory
+    from prompt_toolkit.auto_suggest import AutoSuggest, Suggestion
+    from prompt_toolkit.key_binding import KeyBindings
+    from prompt_toolkit.styles import Style
+    from prompt_toolkit.patch_stdout import patch_stdout
+    from prompt_toolkit import print_formatted_text
+    from prompt_toolkit.formatted_text import ANSI
+    HAS_PROMPT_TOOLKIT = True
+except ImportError as e:
+    HAS_PROMPT_TOOLKIT = False
+    _PT_IMPORT_ERROR = str(e)
+
+
+if HAS_PROMPT_TOOLKIT:
+
+    class UECompleter(Completer):
+        """Server-assisted completer with caching."""
+
+        PAGE_SIZE = 200
+        MIN_PREFIX = 2
+
+        def __init__(self, console: RemoteConsole):
+            self.console = console
+            self._cached_prefix = ""
+            self._cached_results = []
+            self._cached_total = 0
+            self._cached_time = 0.0
+            self._cache_ttl = 10.0
+
+        def _from_cache(self, prefix: str):
+            now = time.monotonic()
+            if not self._cached_prefix or (now - self._cached_time) > self._cache_ttl:
+                return None
+            # Don't serve from cache if the cached result was empty (server may not have been ready)
+            if not self._cached_results:
+                return None
+            if (prefix.startswith(self._cached_prefix)
+                    and self._cached_total <= len(self._cached_results)):
+                pl = prefix.lower()
+                return [(n, v, s, h) for n, v, s, h in self._cached_results
+                        if n.lower().startswith(pl)]
+            return None
+
+        def _fetch(self, prefix: str):
+            try:
+                results, total = self.console.complete(prefix, offset=0, limit=self.PAGE_SIZE)
+            except (socket.timeout, ConnectionError, OSError, TimeoutError):
+                return []
+            self._cached_prefix = prefix
+            self._cached_results = results
+            self._cached_total = total
+            self._cached_time = time.monotonic()
+            return list(results)  # copy: caller may sort in-place
+
+        def get_completions(self, document, complete_event):
+            text = document.text_before_cursor.lstrip()
+            if not text:
+                return
+
+            # Internal /commands
+            if text.startswith("/"):
+                typed = text.lower()
+                for name, desc in INTERNAL_COMMANDS_COMMON.items():
+                    if name.startswith(typed):
+                        yield Completion(name, start_position=-len(text),
+                                         display_meta=desc)
+                return
+
+            if len(text) < self.MIN_PREFIX:
+                return
+
+            results = self._from_cache(text)
+            if results is None:
+                results = self._fetch(text)
+
+            # Sort: shorter names first (so "r.Nanite" comes before
+            # "r.Nanite.AllowMaskedMaterials"), then alphabetically.
+            results.sort(key=lambda r: (len(r[0]), r[0].lower()))
+
+            for name, value, source, help_text in results:
+                parts = []
+                if value:
+                    parts.append(f"={value}")
+                if source:
+                    parts.append(f"({source})")
+                if help_text:
+                    parts.append(help_text)
+                meta = " ".join(parts)
+                yield Completion(name, start_position=-len(text), display_meta=meta)
+
+    class UEAutoSuggest(AutoSuggest):
+        """Ghost text from the first completion match."""
+
+        def __init__(self, completer: UECompleter):
+            self._completer = completer
+
+        def get_suggestion(self, buffer, document):
+            text = document.text_before_cursor.lstrip()
+            if not text or len(text) < 3:
+                return None
+            # Use cached results only - never block the UI with a network call.
+            results = self._completer._from_cache(text)
+            if results is None:
+                return None
+            for name, _, _, _ in results:
+                if name.lower().startswith(text.lower()) and len(name) > len(text):
+                    return Suggestion(name[len(text):])
+            return None
+
+
+def _interactive_loop_rich(console: RemoteConsole):
+    """Rich interactive mode with prompt_toolkit."""
+
+    # Build prompt string: "TargetName> " or "ue@host:port> "
+    if console.label:
+        prompt_str = f"{console.label}> "
+    else:
+        prompt_str = f"ue@{console.host}:{console.port}> "
+
+    reconnected_event = threading.Event()  # signaled when connection is restored
+    is_disconnected = threading.Event()    # stays set while disconnected
+
+    completer = UECompleter(console)
+    threaded_completer = ThreadedCompleter(completer)
+    history_file = os.path.expanduser("~/.ue_remote_console_history")
+
+    style = Style.from_dict({
+        "completion-menu.completion":             "bg:#333333 #ffffff",
+        "completion-menu.completion.current":     "bg:#555555 #ffffff bold",
+        "completion-menu.meta.completion":        "bg:#333333 #aaaaaa",
+        "completion-menu.meta.completion.current": "bg:#555555 #cccccc",
+        "auto-suggest":                           "#666666",
+    })
+
+    bindings = KeyBindings()
+
+    @bindings.add("escape")
+    def _on_escape(event):
+        buf = event.app.current_buffer
+        if buf.complete_state:
+            buf.cancel_completion()
+
+    @bindings.add("pagedown")
+    def _on_pagedown(event):
+        buf = event.app.current_buffer
+        if buf.complete_state:
+            for _ in range(10):
+                buf.complete_next()
+
+    @bindings.add("pageup")
+    def _on_pageup(event):
+        buf = event.app.current_buffer
+        if buf.complete_state:
+            for _ in range(10):
+                buf.complete_previous()
+
+    @bindings.add("c-l")
+    def _on_ctrl_l(event):
+        buf = event.app.current_buffer
+        if buf.text:
+            buf.reset()
+        else:
+            # Clear screen and reposition cursor to top, then invalidate
+            # the renderer so prompt_toolkit redraws the prompt line.
+            event.app.renderer.output.write_raw("\033[2J\033[H")
+            event.app.renderer.output.flush()
+            event.app.renderer.reset()
+            event.app.invalidate()
+
+    @bindings.add("c-h")  # Backspace sends ^H
+    @bindings.add("backspace")
+    def _on_backspace(event):
+        buf = event.app.current_buffer
+        if buf.document.cursor_position > 0:
+            buf.delete_before_cursor()
+        # else: do nothing (no bell)
+
+    pending_dangerous = ""
+
+    session = PromptSession(
+        history=FileHistory(history_file),
+        auto_suggest=UEAutoSuggest(completer),
+        completer=threaded_completer,
+        complete_while_typing=True,
+        style=style,
+        key_bindings=bindings,
+        reserve_space_for_menu=12,
+        enable_system_prompt=False,
+        enable_open_in_editor=False,
+    )
+
+    def _on_disconnect():
+        is_disconnected.set()
+        reconnected_event.clear()
+        # Flush any buffered log lines so they appear before the disconnect message.
+        _flush_log_buf()
+        print("  [disconnected -- waiting for reconnection...]")
+        # Abort the active prompt so the user can't type into a dead connection.
+        # session.app is only set while prompt() is blocking.
+        app = session.app
+        if app:
+            app.exit(result="")
+
+    def _on_reconnect():
+        is_disconnected.clear()
+        reconnected_event.set()
+        print("  [reconnected]")
+
+    console.enable_auto_reconnect(on_disconnect=_on_disconnect, on_reconnect=_on_reconnect)
+
+    print("Type console commands. Tab: complete. Up/Down: history. /help: commands. Ctrl+D to exit.")
+    print()
+
+    console.set_log_callback(None)  # placeholder, set inside patch_stdout below
+
+    _log_flush_stop = threading.Event()
+    _log_flush_thread = None
+    try:
+      with patch_stdout(raw=True):
+        # Set the log callback *inside* patch_stdout so that print_formatted_text
+        # uses the patched output and renders ANSI colors correctly.
+        #
+        # Log lines are buffered and flushed in batches to avoid the overhead
+        # of per-line print_formatted_text() calls through prompt_toolkit's
+        # rendering pipeline. This makes a large difference for commands that
+        # produce hundreds of log lines (e.g., ProfileGPU: 589 lines went
+        # from >2s to ~0.3s).
+        _log_buf = []
+        _log_buf_lock = threading.Lock()
+
+        def _flush_log_buf():
+            """Write all buffered log lines to the terminal in one shot."""
+            with _log_buf_lock:
+                if not _log_buf:
+                    return
+                batch = "".join(_log_buf)
+                _log_buf.clear()
+            # Write directly to stdout -- the ANSI escapes are already
+            # embedded in the formatted strings. sys.stdout is patched by
+            # patch_stdout() so this cooperates with prompt_toolkit.
+            sys.stdout.write(batch)
+            sys.stdout.flush()
+
+        def _on_log(msg):
+            if is_disconnected.is_set():
+                return  # Discard logs after disconnect to avoid stale output
+            formatted = f"  {_format_log_msg(msg)}\n"
+            with _log_buf_lock:
+                _log_buf.append(formatted)
+
+        # Background thread to flush log buffer at regular intervals.
+        # 50ms gives a good balance between responsiveness and batching.
+        def _log_flush_loop():
+            while not _log_flush_stop.is_set():
+                _flush_log_buf()
+                _log_flush_stop.wait(0.05)
+            _flush_log_buf()  # final flush
+
+        _log_flush_thread = threading.Thread(target=_log_flush_loop, daemon=True)
+        _log_flush_thread.start()
+        console.set_log_callback(_on_log)
+        while True:
+            # Block while disconnected - don't show input prompt until connected.
+            # Use short sleeps so Ctrl+C can be delivered on Windows (Event.wait
+            # swallows KeyboardInterrupt on some Python/OS combinations).
+            if is_disconnected.is_set():
+                try:
+                    while not reconnected_event.is_set():
+                        if not is_disconnected.is_set():
+                            break
+                        time.sleep(0.2)
+                except KeyboardInterrupt:
+                    print()
+                    break
+                reconnected_event.clear()
+                continue
+
+            try:
+                cmd = session.prompt(prompt_str).strip()
+            except (EOFError, KeyboardInterrupt):
+                print()
+                break
+
+            if not cmd:
+                continue
+
+            # Reject input if we lost connection between prompt display and submit.
+            if is_disconnected.is_set():
+                print("  [not connected]")
+                continue
+
+            if cmd.startswith("/"):
+                if not _handle_internal_command(cmd, console, rich_mode=True):
+                    break
+                continue
+
+            # Dangerous command confirmation
+            first_word = cmd.split()[0].lower()
+            if first_word in DANGEROUS_COMMANDS:
+                if cmd == pending_dangerous:
+                    pending_dangerous = ""
+                else:
+                    pending_dangerous = cmd
+                    print(f"  WARNING: '{cmd}' may disrupt the target. Enter same command again to confirm.")
+                    print()
+                    continue
+            else:
+                pending_dangerous = ""
+
+            try:
+                start = time.monotonic()
+                resp = console.execute(cmd)
+                elapsed = time.monotonic() - start
+                if resp.get("error"):
+                    print(f"  Error: {resp['error']}")
+                print(f"  ({elapsed:.3f}s)")
+            except TimeoutError as e:
+                print(f"  Timed out: {e} (server may be paused on a breakpoint)", file=sys.stderr)
+            except ConnectionError as e:
+                print(f"  Connection lost: {e}", file=sys.stderr)
+                print("  Waiting for reconnection (Ctrl+D to exit)...", file=sys.stderr)
+                # Ensure disconnected state is visible even if the callback
+                # hasn't fired yet (race between reader thread and main thread).
+                is_disconnected.set()
+    finally:
+        _log_flush_stop.set()
+        if _log_flush_thread is not None:
+            _log_flush_thread.join(timeout=1.0)
+
+
+
+
+def _interactive_loop_basic(console: RemoteConsole):
+    """Simple input() REPL with no TUI."""
+
+    # Build prompt string
+    if console.label:
+        prompt_str = f"{console.label}> "
+    else:
+        prompt_str = f"ue@{console.host}:{console.port}> "
+
+    reconnected_event = threading.Event()
+    is_disconnected = threading.Event()
+
+    def _on_disconnect():
+        is_disconnected.set()
+        reconnected_event.clear()
+        print("  [disconnected -- waiting for reconnection...]", file=sys.stderr)
+
+    def _on_reconnect():
+        is_disconnected.clear()
+        reconnected_event.set()
+        print("  [reconnected]")
+
+    console.enable_auto_reconnect(on_disconnect=_on_disconnect, on_reconnect=_on_reconnect)
+
+    print("Type console commands. /help for internal commands. Ctrl+C to exit.")
+    print()
+
+    # Print log lines directly from the callback thread.
+    # Basic mode uses input() which doesn't control the terminal,
+    # so interleaved output is acceptable.
+    def _on_log(msg):
+        print(f"  {_format_log_msg(msg)}")
+
+    console.set_log_callback(_on_log)
+
+    while True:
+        # Block while disconnected - don't show input prompt until connected.
+        if is_disconnected.is_set():
+            try:
+                while not reconnected_event.is_set():
+                    if not is_disconnected.is_set():
+                        break
+                    time.sleep(0.2)
+            except KeyboardInterrupt:
+                print()
+                break
+            reconnected_event.clear()
+            continue
+
+        try:
+            cmd = input(prompt_str).strip()
+        except (EOFError, KeyboardInterrupt):
+            print()
+            break
+
+        if not cmd:
+            continue
+
+        # Reject input if we lost connection between prompt display and submit.
+        if is_disconnected.is_set():
+            print("  [not connected]")
+            continue
+
+        try:
+            if cmd.startswith("/"):
+                if not _handle_internal_command(cmd, console):
+                    break
+            else:
+                first_word = cmd.split()[0].lower()
+                if first_word in DANGEROUS_COMMANDS:
+                    try:
+                        confirm = input(
+                            f"  Send '{cmd}' to remote? This may disrupt the target. [y/N] "
+                        ).strip().lower()
+                    except (EOFError, KeyboardInterrupt):
+                        print()
+                        continue
+                    if confirm not in ("y", "yes"):
+                        print("  (cancelled)")
+                        print()
+                        continue
+                start = time.monotonic()
+                resp = console.execute(cmd)
+                elapsed = time.monotonic() - start
+                if resp.get("error"):
+                    print(f"  Error: {resp['error']}")
+                print(f"  ({elapsed:.3f}s)")
+        except TimeoutError as e:
+            print(f"  Timed out: {e} (server may be paused on a breakpoint)", file=sys.stderr)
+        except ConnectionError as e:
+            print(f"  Connection lost: {e}", file=sys.stderr)
+            print("  Waiting for reconnection (Ctrl+C to exit)...", file=sys.stderr)
+            is_disconnected.set()
+
+
+def interactive_loop(console: RemoteConsole):
+    """Interactive REPL - rich mode with prompt_toolkit, or basic fallback."""
+    if HAS_PROMPT_TOOLKIT:
+        _interactive_loop_rich(console)
+    else:
+        err = globals().get("_PT_IMPORT_ERROR", "not installed")
+        print("Tip: Install 'prompt_toolkit' for auto-complete, ghost text, and history:",
+              file=sys.stderr)
+        print("  pip install prompt_toolkit", file=sys.stderr)
+        print(f"  (import error: {err})", file=sys.stderr)
+        print(file=sys.stderr)
+        _interactive_loop_basic(console)
+
+
+def pipe_mode(console: RemoteConsole):
+    """Read commands from stdin (non-interactive, for scripting)."""
+    for raw_line in sys.stdin:
+        cmd = raw_line.strip()
+        if not cmd or cmd.startswith("#"):
+            continue
+        try:
+            resp = console.execute(cmd)
+            if resp.get("error"):
+                print(f"Error: {resp['error']}")
+        except (ConnectionError, TimeoutError) as e:
+            print(f"Error: {e}", file=sys.stderr)
+            break
diff --git a/Engine/Extras/UnrealRemoteConsole/remote_console/tests.py b/Engine/Extras/UnrealRemoteConsole/remote_console/tests.py
new file mode 100644
--- /dev/null
+++ b/Engine/Extras/UnrealRemoteConsole/remote_console/tests.py
@@ -0,0 +1,659 @@
+"""Self-tests for the UE Remote Console client."""
+
+import codecs
+import json
+import re
+import socket
+import sys
+import threading
+import time
+
+from .client import RemoteConsole, DEFAULT_TIMEOUT, DANGEROUS_COMMANDS
+from .formatting import (
+    _format_log_msg, _highlight_log_text, _disable_colors,
+    _NUM_COLOR, _UNIT_COLOR, _STR_COLOR, _QUOTE_COLOR, _ARROW_COLOR,
+    _BRACKET_COLOR, _KEY_COLOR, _RESET,
+)
+
+
+def _run_self_tests():
+    """Run internal tests to verify UTF-8 decoding and protocol handling."""
+    passed = 0
+    failed = 0
+
+    def _check(name, condition, detail=""):
+        nonlocal passed, failed
+        if condition:
+            passed += 1
+            print(f"  PASS  {name}")
+        else:
+            failed += 1
+            print(f"  FAIL  {name}  {detail}")
+
+    def _make_log_line(text, ascii_only=False):
+        obj = {"op": "log", "cat": "Test", "line": text}
+        return (json.dumps(obj, ensure_ascii=ascii_only) + "\n").encode("utf-8")
+
+    def _decode_chunks(chunks):
+        """Simulate the reader's incremental decode pipeline."""
+        decoder = codecs.getincrementaldecoder("utf-8")("replace")
+        text_buf = ""
+        lines = []
+        for chunk in chunks:
+            text_buf += decoder.decode(chunk)
+            while "\n" in text_buf:
+                line, text_buf = text_buf.split("\n", 1)
+                line = line.rstrip("\r")
+                if line:
+                    lines.append(line)
+        remaining = decoder.decode(b"", final=True)
+        if remaining:
+            text_buf += remaining
+        if text_buf.strip():
+            lines.append(text_buf.strip())
+        return lines
+
+    def _extract(json_line):
+        return json.loads(json_line).get("line", "")
+
+    print("Running self-tests...")
+    print()
+
+    # --- UTF-8 boundary split tests ---
+    print("UTF-8 incremental decoder:")
+
+    # Baseline: no split
+    raw = _make_log_line("abc \u2502 def \u2503 ghi")
+    lines = _decode_chunks([raw])
+    _check("no split", len(lines) == 1 and _extract(lines[0]) == "abc \u2502 def \u2503 ghi")
+
+    # Split before multi-byte char
+    raw = _make_log_line("A\u2503B")
+    idx = raw.index(b"\xe2\x94\x83")
+    lines = _decode_chunks([raw[:idx], raw[idx:]])
+    _check("split before 3-byte char", len(lines) == 1 and _extract(lines[0]) == "A\u2503B")
+
+    # Split after 1st byte of 3-byte sequence (E2 | 94 83)
+    lines = _decode_chunks([raw[:idx + 1], raw[idx + 1:]])
+    _check("split after byte 1 of 3", len(lines) == 1 and _extract(lines[0]) == "A\u2503B")
+
+    # Split after 2nd byte of 3-byte sequence (E2 94 | 83)
+    lines = _decode_chunks([raw[:idx + 2], raw[idx + 2:]])
+    _check("split after byte 2 of 3", len(lines) == 1 and _extract(lines[0]) == "A\u2503B")
+
+    # Each byte of 3-byte char in separate chunk
+    lines = _decode_chunks([raw[:idx + 1], raw[idx + 1:idx + 2], raw[idx + 2:]])
+    _check("3-byte char across 3 chunks", len(lines) == 1 and _extract(lines[0]) == "A\u2503B")
+
+    # 4-byte emoji split at each position
+    emoji = "\U0001F600"
+    raw = _make_log_line(f"x{emoji}y")
+    emoji_bytes = emoji.encode("utf-8")  # F0 9F 98 80
+    idx = raw.index(emoji_bytes)
+    all_ok = True
+    for sp in range(1, 4):
+        lines = _decode_chunks([raw[:idx + sp], raw[idx + sp:]])
+        if not (len(lines) == 1 and _extract(lines[0]) == f"x{emoji}y"):
+            all_ok = False
+            break
+    _check("4-byte emoji split at each position", all_ok)
+
+    # Byte-at-a-time delivery
+    raw = _make_log_line("\u2502\u2503\u2502\u2503\u2502")
+    chunks = [raw[i:i + 1] for i in range(len(raw))]
+    lines = _decode_chunks(chunks)
+    _check("byte-at-a-time delivery",
+           len(lines) == 1 and _extract(lines[0]) == "\u2502\u2503\u2502\u2503\u2502")
+
+    # Multiple lines split mid-char
+    line1 = _make_log_line("line1")
+    line2 = _make_log_line("A\u2503B")
+    raw = line1 + line2
+    idx = raw.index(b"\xe2\x94\x83")
+    lines = _decode_chunks([raw[:idx + 1], raw[idx + 1:]])
+    _check("two lines, split mid-char in second",
+           len(lines) == 2 and _extract(lines[1]) == "A\u2503B")
+
+    # ASCII-escaped JSON: immune to any split
+    print()
+    print("ASCII-escaped JSON (server-side \\\\uXXXX):")
+    raw = _make_log_line("\u2502 data \u2503 more \u2502", ascii_only=True)
+    all_ok = True
+    for i in range(1, len(raw)):
+        lines = _decode_chunks([raw[:i], raw[i:]])
+        if not (len(lines) == 1 and _extract(lines[0]) == "\u2502 data \u2503 more \u2502"):
+            all_ok = False
+            break
+    _check("split at every byte position (pure ASCII)", all_ok)
+
+    # --- Mock server round-trip ---
+    print()
+    print("Mock server round-trip:")
+
+    text = "0.4% \u250a          \u2502 0.000 ms \u2503"
+    raw = _make_log_line(text)
+    # Split inside the first box char
+    idx = raw.index(b"\xe2")
+    chunk1, chunk2 = raw[:idx + 1], raw[idx + 1:]
+
+    received = []
+    server = socket.socket(socket.AF_INET, socket.SOCK_STREAM)
+    server.setsockopt(socket.SOL_SOCKET, socket.SO_REUSEADDR, 1)
+    server.bind(("127.0.0.1", 0))
+    server.listen(1)
+    port = server.getsockname()[1]
+
+    def _mock_handler():
+        server.settimeout(5.0)
+        conn, _ = server.accept()
+        try:
+            conn.settimeout(2.0)
+            # Read the handshake hello message
+            hello_buf = b""
+            while b"\n" not in hello_buf:
+                hello_buf += conn.recv(4096)
+            hello_msg = json.loads(hello_buf.decode("utf-8"))
+            resp = json.dumps({"id": hello_msg["id"], "op": "hello", "version": 1}) + "\n"
+            conn.sendall(resp.encode("utf-8"))
+            conn.sendall(chunk1)
+            time.sleep(0.05)
+            conn.sendall(chunk2)
+            time.sleep(0.5)
+        finally:
+            conn.close()
+            server.close()
+
+    server_thread = threading.Thread(target=_mock_handler, daemon=True)
+    server_thread.start()
+
+    try:
+        console = RemoteConsole("127.0.0.1", port, timeout=5.0)
+        console.set_log_callback(lambda msg: received.append(msg.get("line", "")))
+        connected = console.connect()
+        if connected:
+            deadline = time.monotonic() + 3.0
+            while not received and time.monotonic() < deadline:
+                time.sleep(0.05)
+            console.close(permanent=True)
+
+        _check("connect to mock server", connected)
+        _check("received log line", len(received) == 1,
+               f"got {len(received)} lines" if received else "no lines received")
+        if received:
+            _check("text matches exactly", received[0] == text,
+                   f"got: {repr(received[0])}")
+            _check("no replacement chars", "\ufffd" not in received[0])
+    except Exception as e:
+        _check("mock server test", False, str(e))
+
+    server_thread.join(timeout=3.0)
+
+    # --- Syntax highlighting tests ---
+    print()
+    print("Syntax highlighting:")
+
+    def _strip_ansi(s):
+        """Remove ANSI escape sequences to get plain text."""
+        return re.sub(r'\033\[[0-9;]*m', '', s)
+
+    def _has_ansi(s, substr):
+        """Check that substr in s is wrapped in ANSI codes (i.e., highlighted)."""
+        stripped = _strip_ansi(s)
+        idx = stripped.find(substr)
+        if idx < 0:
+            return False
+        # Walk the original string to find where substr starts
+        plain_pos = 0
+        orig_pos = 0
+        while plain_pos < idx and orig_pos < len(s):
+            if s[orig_pos] == '\033':
+                while orig_pos < len(s) and s[orig_pos] != 'm':
+                    orig_pos += 1
+                orig_pos += 1  # skip 'm'
+            else:
+                plain_pos += 1
+                orig_pos += 1
+        # Check if ANSI codes exist right at the start of substr
+        pos_before = orig_pos
+        while orig_pos < len(s) and s[orig_pos] == '\033':
+            while orig_pos < len(s) and s[orig_pos] != 'm':
+                orig_pos += 1
+            orig_pos += 1
+        return orig_pos > pos_before
+
+    def _format_test(msg_dict):
+        """Run _format_log_msg and return (formatted, stripped)."""
+        fmt = _format_log_msg(msg_dict)
+        return fmt, _strip_ansi(fmt)
+
+    # Test: category gets colored
+    fmt, plain = _format_test({"cat": "LogEngine", "v": "Log", "line": "Hello"})
+    _check("category [LogEngine] colored", _has_ansi(fmt, "[LogEngine]"))
+    _check("category preserved in plain text", "[LogEngine] Hello" == plain)
+
+    # Test: deterministic coloring (CRC32-based)
+    fmt1, _ = _format_test({"cat": "LogRHI", "v": "Log", "line": "A"})
+    fmt2, _ = _format_test({"cat": "LogRHI", "v": "Log", "line": "B"})
+    # Extract the color code for [LogRHI] from both
+    m1 = re.search(r'(\033\[[0-9;]*m)\[LogRHI\]', fmt1)
+    m2 = re.search(r'(\033\[[0-9;]*m)\[LogRHI\]', fmt2)
+    _check("category color is deterministic",
+           m1 and m2 and m1.group(1) == m2.group(1),
+           f"got {m1.group(1) if m1 else 'None'} vs {m2.group(1) if m2 else 'None'}")
+
+    # Test: brackets [[ and ]] highlighted
+    fmt, plain = _format_test({"cat": "LogConfig", "v": "Log",
+                                "line": "CVar [[r.HZB.BuildUseCompute:1 -> 0]]"})
+    _check("double brackets [[ highlighted", _has_ansi(fmt, "[["))
+    _check("double brackets ]] highlighted", _has_ansi(fmt, "]]"))
+    _check("no text eaten by brackets",
+           plain == "[LogConfig] CVar [[r.HZB.BuildUseCompute:1 -> 0]]")
+
+    # Test: key:value colored
+    fmt, plain = _format_test({"cat": "Log", "v": "Log",
+                                "line": "[[rhi.syncinterval:1 -> 0]]"})
+    _check("key 'rhi.syncinterval' colored", _has_ansi(fmt, "rhi.syncinterval"))
+    _check("separator ':' colored", _has_ansi(fmt, ":"))
+    _check("number '1' colored", _has_ansi(fmt, "1"))
+    _check("arrow '->' colored", _has_ansi(fmt, "->"))
+    _check("no text eaten by key:value",
+           plain == "[Log] [[rhi.syncinterval:1 -> 0]]")
+
+    # Test: key=value colored
+    fmt, plain = _format_test({"cat": "Log", "v": "Log", "line": "Width=1920"})
+    _check("key 'Width' colored in key=value", _has_ansi(fmt, "Width"))
+    _check("number '1920' colored in key=value", _has_ansi(fmt, "1920"))
+
+    # Test: quoted strings colored
+    fmt, plain = _format_test({"cat": "Log", "v": "Log",
+                                "line": 'r.Nanite = "1"'})
+    _check("quoted string '\"1\"' colored", _has_ansi(fmt, '"1"'))
+    _check("no text eaten by string highlight", plain == '[Log] r.Nanite = "1"')
+
+    # Test: numbers NOT highlighted inside identifiers
+    # Use _highlight_log_text directly to avoid category prefix interference
+    hl = _highlight_log_text("SpotLight77 Phase 2")
+    hl_plain = _strip_ansi(hl)
+    _check("number in identifier not highlighted",
+           not _has_ansi(hl, "SpotLight77"))
+    # Verify the raw highlighted output contains \033[36m2\033[0m (cyan number)
+    _check("standalone number highlighted", f"{_NUM_COLOR}2{_RESET}" in hl)
+    _check("identifier+number text preserved", hl_plain == "SpotLight77 Phase 2")
+
+    # Test: parentheses highlighted
+    fmt, _ = _format_test({"cat": "Log", "v": "Log", "line": "(155K)"})
+    _check("opening paren highlighted", _has_ansi(fmt, "("))
+    _check("closing paren highlighted", _has_ansi(fmt, ")"))
+
+    # Test: verbosity label shown for errors
+    fmt, plain = _format_test({"cat": "LogMat", "v": "Error",
+                                "line": "Shader not found"})
+    _check("error verbosity label shown", "Error:" in plain)
+    _check("error text preserved", "Shader not found" in plain)
+
+    # Test: verbosity label NOT shown for normal logs
+    fmt, plain = _format_test({"cat": "LogTemp", "v": "Log", "line": "Hello"})
+    _check("no verbosity label for Log level", "Log:" not in plain)
+
+    # Test: leading [Tag] in line text colored
+    fmt, plain = _format_test({"cat": "LogEngine", "v": "Log",
+                                "line": "[SubSystem] Initialized"})
+    _check("leading [SubSystem] tag colored", _has_ansi(fmt, "[SubSystem]"))
+    _check("leading tag text preserved", "[SubSystem] Initialized" in plain)
+
+    # Test: angle brackets highlighted
+    fmt, _ = _format_test({"cat": "Log", "v": "Log",
+                            "line": "<Console variable> = <value>"})
+    _check("angle bracket < highlighted", _has_ansi(fmt, "<"))
+    _check("angle bracket > highlighted", _has_ansi(fmt, ">"))
+
+    # Test: standalone = highlighted (not part of key=value)
+    _check("standalone = highlighted", _has_ansi(fmt, "="))
+
+    # Test: number with unit suffix -- split coloring
+    hl = _highlight_log_text("size 512KB rate 2.4MB/s time 3.5ms freq 60Hz pct 99.1%")
+    hl_plain = _strip_ansi(hl)
+    _check("unit text preserved",
+           hl_plain == "size 512KB rate 2.4MB/s time 3.5ms freq 60Hz pct 99.1%")
+    # Number portion in bright cyan, suffix in dim cyan
+    _check("512 number colored", f"{_NUM_COLOR}512{_RESET}" in hl)
+    _check("KB suffix colored differently", f"{_UNIT_COLOR}KB{_RESET}" in hl)
+    _check("2.4 number colored", f"{_NUM_COLOR}2.4{_RESET}" in hl)
+    _check("MB/s suffix colored", f"{_UNIT_COLOR}MB/s{_RESET}" in hl)
+    _check("3.5 number colored", f"{_NUM_COLOR}3.5{_RESET}" in hl)
+    _check("ms suffix colored", f"{_UNIT_COLOR}ms{_RESET}" in hl)
+    _check("60 number colored", f"{_NUM_COLOR}60{_RESET}" in hl)
+    _check("Hz suffix colored", f"{_UNIT_COLOR}Hz{_RESET}" in hl)
+    _check("99.1 number colored", f"{_NUM_COLOR}99.1{_RESET}" in hl)
+    _check("% suffix colored", f"{_UNIT_COLOR}%{_RESET}" in hl)
+
+    # Test: detached unit suffix (space between number and unit)
+    hl = _highlight_log_text("freed 2.4 MB, peak 1 GiB")
+    _check("2.4 with space+MB colored", f"{_NUM_COLOR}2.4{_RESET}" in hl)
+    _check("space+MB suffix colored", f"{_UNIT_COLOR} MB{_RESET}" in hl)
+    _check("1 with space+GiB colored", f"{_NUM_COLOR}1{_RESET}" in hl)
+    _check("space+GiB suffix colored", f"{_UNIT_COLOR} GiB{_RESET}" in hl)
+
+    # Test: single-quoted string with dim quotes, bright content
+    hl = _highlight_log_text("loaded 'foo.pak'")
+    _check("single quote mark dim", f"{_QUOTE_COLOR}'{_RESET}" in hl)
+    _check("single quote content bright", f"{_STR_COLOR}foo.pak{_RESET}" in hl)
+
+    # Test: double-quoted string with dim quotes, bright content
+    hl = _highlight_log_text('value "hello"')
+    _check("double quote mark dim", f'{_QUOTE_COLOR}"{_RESET}' in hl)
+    _check("double quote content bright", f"{_STR_COLOR}hello{_RESET}" in hl)
+
+    # Test: warning verbosity
+    fmt, plain = _format_test({"cat": "Log", "v": "Warning", "line": "Low memory"})
+    _check("warning verbosity label shown", "Warning:" in plain)
+
+    # Test: fatal verbosity
+    fmt, plain = _format_test({"cat": "Log", "v": "Fatal", "line": "Crash"})
+    _check("fatal verbosity label shown", "Fatal:" in plain)
+
+    # Test: time units ns, us, s and px (pixels)
+    hl = _highlight_log_text("latency 50ns avg 1.2us total 1.5s offset 320px")
+    _check("ns suffix colored", f"{_UNIT_COLOR}ns{_RESET}" in hl)
+    _check("us suffix colored", f"{_UNIT_COLOR}us{_RESET}" in hl)
+    _check("s suffix colored", f"{_UNIT_COLOR}s{_RESET}" in hl)
+    _check("px suffix colored", f"{_UNIT_COLOR}px{_RESET}" in hl)
+
+    # Test: arrow -> highlighted
+    hl = _highlight_log_text("value 1 -> 0")
+    _check("arrow -> highlighted", f"{_ARROW_COLOR}->{_RESET}" in hl)
+
+    # Test: numbers inside quoted strings are highlighted
+    hl = _highlight_log_text('systemresolution.resx="1920"')
+    _check("number inside double-quoted string colored",
+           f"{_NUM_COLOR}1920{_RESET}" in hl)
+    _check("quote marks still dim around number",
+           f"{_QUOTE_COLOR}\"{_RESET}" in hl)
+    hl = _highlight_log_text("size='512KB'")
+    _check("number inside single-quoted string colored",
+           f"{_NUM_COLOR}512{_RESET}" in hl)
+    _check("unit suffix inside quoted string colored",
+           f"{_UNIT_COLOR}KB{_RESET}" in hl)
+
+    # Test: dimension patterns (NUMxNUM)
+    hl = _highlight_log_text("Resolution 1920x1080 Volume 256x256x4")
+    hl_plain = _strip_ansi(hl)
+    _check("dimension text preserved",
+           hl_plain == "Resolution 1920x1080 Volume 256x256x4")
+    _check("dimension number colored", f"{_NUM_COLOR}1920{_RESET}" in hl)
+    _check("dimension x uses arrow color", f"{_ARROW_COLOR}x{_RESET}" in hl)
+    _check("dimension second number colored", f"{_NUM_COLOR}1080{_RESET}" in hl)
+
+    # Test: dimensions inside quoted strings
+    hl = _highlight_log_text('Foo="123x345"')
+    hl_plain = _strip_ansi(hl)
+    _check("dimension inside string text preserved",
+           hl_plain == 'Foo="123x345"')
+    _check("dimension inside string number colored",
+           f"{_NUM_COLOR}123{_RESET}" in hl)
+    _check("dimension inside string x colored",
+           f"{_ARROW_COLOR}x{_RESET}" in hl)
+    _check("dimension inside string second number colored",
+           f"{_NUM_COLOR}345{_RESET}" in hl)
+
+    # Test: :: is NOT treated as key:value (C++ scope resolution)
+    hl = _highlight_log_text("Substrate::MaterialClassification")
+    hl_plain = _strip_ansi(hl)
+    _check(":: not treated as key:value",
+           not _has_ansi(hl, "Substrate"))
+    _check(":: text preserved",
+           hl_plain == "Substrate::MaterialClassification")
+
+    # Test: key: value with space after separator IS treated as key-value
+    hl = _highlight_log_text("Textures: 0, Buffers: 1")
+    _check("'Textures' colored as key with space",
+           _has_ansi(hl, "Textures"))
+    _check("'Buffers' colored as key with space",
+           _has_ansi(hl, "Buffers"))
+
+    # Test: key = "value" with spaces around = (CVar assignment style)
+    hl = _highlight_log_text('r.DynamicRes.TestScreenPercentage = "40"')
+    hl_plain = _strip_ansi(hl)
+    _check("spaced = key colored",
+           _has_ansi(hl, "r.DynamicRes.TestScreenPercentage"))
+    _check("spaced = separator colored",
+           _has_ansi(hl, " = "))
+    _check("spaced = text preserved",
+           hl_plain == 'r.DynamicRes.TestScreenPercentage = "40"')
+
+    # Test: key= "value" (no space before =, space after)
+    hl = _highlight_log_text('r.Foo= "40"')
+    _check("key= value key colored", _has_ansi(hl, "r.Foo"))
+    _check("key= value text preserved", _strip_ansi(hl) == 'r.Foo= "40"')
+
+    # Test: key ="value" (space before =, no space after)
+    hl = _highlight_log_text("r.Foo ='40'")
+    _check("key =value key colored", _has_ansi(hl, "r.Foo"))
+
+    # --- Commands and interactive module tests ---
+    # These catch import/wiring errors that would only surface at runtime.
+    print()
+    print("Commands and interactive:")
+
+    from .commands import (
+        _handle_internal_command, INTERNAL_COMMANDS_COMMON,
+        INTERNAL_COMMANDS_BASIC_ONLY, INTERNAL_COMMANDS,
+        INTERNAL_HELP_TEXT_COMMON, INTERNAL_HELP_TEXT_BASIC,
+    )
+
+    # Verify command dictionaries are consistent
+    _check("INTERNAL_COMMANDS is superset of COMMON",
+           all(k in INTERNAL_COMMANDS for k in INTERNAL_COMMANDS_COMMON))
+    _check("INTERNAL_COMMANDS is superset of BASIC_ONLY",
+           all(k in INTERNAL_COMMANDS for k in INTERNAL_COMMANDS_BASIC_ONLY))
+    _check("INTERNAL_COMMANDS size matches",
+           len(INTERNAL_COMMANDS) == len(INTERNAL_COMMANDS_COMMON) + len(INTERNAL_COMMANDS_BASIC_ONLY))
+    _check("all commands start with /",
+           all(k.startswith("/") for k in INTERNAL_COMMANDS))
+    _check("help text mentions /help",
+           "/help" in INTERNAL_HELP_TEXT_COMMON and "/help" in INTERNAL_HELP_TEXT_BASIC)
+
+    # Mock console for command handler tests
+    class _MockConsole:
+        def __init__(self):
+            self.connected = True
+            self.executed = []
+            self.completed = []
+            self.reconnected = False
+        def execute(self, cmd):
+            self.executed.append(cmd)
+            return {"ok": True}
+        def complete(self, prefix, offset=0, limit=100):
+            self.completed.append(prefix)
+            return [("r.Foo", "1", "Console", "A foo")], 1
+        def reconnect(self):
+            self.reconnected = True
+        def getvar(self, name):
+            return {"name": name, "value": "42", "source": "Console", "help": "test"}
+
+    mock = _MockConsole()
+    log_output = []
+
+    # /quit returns False (exit loop)
+    _check("/quit returns False",
+           _handle_internal_command("/quit", mock, log_fn=log_output.append) == False)
+    _check("/exit returns False",
+           _handle_internal_command("/exit", mock, log_fn=log_output.append) == False)
+    _check("/q returns False",
+           _handle_internal_command("/q", mock, log_fn=log_output.append) == False)
+
+    # /help returns True (continue loop) and produces output
+    log_output.clear()
+    _check("/help returns True",
+           _handle_internal_command("/help", mock, log_fn=log_output.append) == True)
+    _check("/help produces output", len(log_output) > 0)
+
+    log_output.clear()
+    _check("/h returns True",
+           _handle_internal_command("/h", mock, log_fn=log_output.append) == True)
+
+    log_output.clear()
+    _check("/? returns True",
+           _handle_internal_command("/?", mock, log_fn=log_output.append) == True)
+
+    # /reconnect triggers reconnect
+    mock.reconnected = False
+    _handle_internal_command("/reconnect", mock, log_fn=log_output.append)
+    _check("/reconnect calls reconnect()", mock.reconnected)
+
+    # /complete with prefix
+    log_output.clear()
+    _handle_internal_command("/complete r.Foo", mock, log_fn=log_output.append)
+    _check("/complete calls complete()", "r.Foo" in mock.completed)
+    _check("/complete produces output", len(log_output) > 0)
+
+    # /complete without args shows usage
+    log_output.clear()
+    _handle_internal_command("/complete", mock, log_fn=log_output.append)
+    _check("/complete no-args shows usage",
+           any("Usage" in str(l) for l in log_output))
+
+    # /get with name
+    log_output.clear()
+    _handle_internal_command("/get r.Foo", mock, log_fn=log_output.append)
+    _check("/get produces output", len(log_output) > 0)
+
+    # /get without args shows usage
+    log_output.clear()
+    _handle_internal_command("/get", mock, log_fn=log_output.append)
+    _check("/get no-args shows usage",
+           any("Usage" in str(l) for l in log_output))
+
+    # Unknown command
+    log_output.clear()
+    _handle_internal_command("/nonexistent", mock, log_fn=log_output.append)
+    _check("/unknown shows error",
+           any("Unknown" in str(l) for l in log_output))
+
+    # Case insensitivity
+    _check("/QUIT returns False",
+           _handle_internal_command("/QUIT", mock, log_fn=log_output.append) == False)
+    _check("/Help returns True",
+           _handle_internal_command("/Help", mock, log_fn=log_output.append) == True)
+
+    # --- Interactive module import validation ---
+    # These verify that all cross-module references resolve correctly,
+    # catching the exact class of bug (missing imports) that broke slash
+    # command completion.
+    print()
+    print("Interactive module imports:")
+
+    from . import interactive as _interactive_mod
+
+    _check("interactive imports INTERNAL_COMMANDS_COMMON",
+           hasattr(_interactive_mod, 'INTERNAL_COMMANDS_COMMON'))
+    _check("interactive imports DANGEROUS_COMMANDS",
+           hasattr(_interactive_mod, 'DANGEROUS_COMMANDS'))
+    _check("interactive imports _handle_internal_command",
+           hasattr(_interactive_mod, '_handle_internal_command'))
+    _check("interactive imports _format_log_msg",
+           hasattr(_interactive_mod, '_format_log_msg'))
+    _check("interactive imports socket",
+           hasattr(_interactive_mod, 'socket'))
+    _check("interactive_loop is callable",
+           callable(getattr(_interactive_mod, 'interactive_loop', None)))
+    _check("pipe_mode is callable",
+           callable(getattr(_interactive_mod, 'pipe_mode', None)))
+
+    # If prompt_toolkit is available, verify the completer classes exist
+    if getattr(_interactive_mod, 'HAS_PROMPT_TOOLKIT', False):
+        _check("UECompleter class exists",
+               hasattr(_interactive_mod, 'UECompleter'))
+        _check("UEAutoSuggest class exists",
+               hasattr(_interactive_mod, 'UEAutoSuggest'))
+
+        # Exercise UECompleter with a mock console
+        Completer_cls = getattr(_interactive_mod, 'UECompleter')
+        completer = Completer_cls(mock)
+
+        # Simulate a Document for completion
+        class _MockDoc:
+            def __init__(self, text):
+                self.text_before_cursor = text
+
+        # Slash command completion
+        slash_results = list(completer.get_completions(_MockDoc("/he"), None))
+        _check("completer returns /help for '/he'",
+               any("/help" in c.text for c in slash_results))
+
+        slash_results = list(completer.get_completions(_MockDoc("/q"), None))
+        _check("completer returns /quit for '/q'",
+               any("/quit" in c.text for c in slash_results))
+
+        # Empty input returns nothing
+        empty_results = list(completer.get_completions(_MockDoc(""), None))
+        _check("completer returns nothing for empty input",
+               len(empty_results) == 0)
+
+        # Short input (< MIN_PREFIX) returns nothing
+        short_results = list(completer.get_completions(_MockDoc("r"), None))
+        _check("completer returns nothing for single char",
+               len(short_results) == 0)
+
+        # Server completion (uses mock)
+        server_results = list(completer.get_completions(_MockDoc("r.Foo"), None))
+        _check("completer fetches server completions",
+               len(server_results) > 0)
+
+    # --- Client module validation ---
+    print()
+    print("Client module:")
+
+    _check("DEFAULT_TIMEOUT is numeric", isinstance(DEFAULT_TIMEOUT, (int, float)))
+    _check("DANGEROUS_COMMANDS is a set", isinstance(DANGEROUS_COMMANDS, set))
+    _check("DANGEROUS_COMMANDS contains exit", "exit" in DANGEROUS_COMMANDS)
+
+    # Verify RemoteConsole can be instantiated without connecting
+    rc = RemoteConsole("127.0.0.1", 0, timeout=1.0)
+    _check("RemoteConsole instantiation", rc is not None)
+    _check("RemoteConsole.host", rc.host == "127.0.0.1")
+    _check("RemoteConsole.port", rc.port == 0)
+    _check("RemoteConsole not connected initially", not rc._connected)
+
+    # --- Visual output samples ---
+    # Print actual colored lines so highlighting can be verified visually.
+    print()
+    print("Visual samples (verify colors manually):")
+    print()
+    samples = [
+        {"cat": "LogEngine", "v": "Log", "line": "Console Help:"},
+        {"cat": "LogConfig", "v": "Log", "line": "CVar [[r.HZB.BuildUseCompute:1 -> 0]] deferred"},
+        {"cat": "LogDeviceProfileManager", "v": "Log", "line": "Pushing Device Profile CVar: [[rhi.syncinterval:1 -> 1]]"},
+        {"cat": "None", "v": "Log", "line": 'r.Nanite = "1"      LastSetBy: Constructor'},
+        {"cat": "LogMaterial", "v": "Error", "line": "Tried to access an uncooked shader map ID in a cooked application"},
+        {"cat": "LogRHI", "v": "Warning", "line": "Shader compilation took 3.5s for 42 shaders"},
+        {"cat": "LogRHI", "v": "Display", "line": "    \u2503     0 \u2502      0 \u2502       0 \u2502       0 \u2502  0.0% \u250a          \u2502 0.001 ms \u2503"},
+        {"cat": "LogTemp", "v": "Log", "line": "[SubSystem] [Init] Loaded 155 assets (2.4 MB) from 'Content/Paks/main.pak'"},
+        {"cat": "LogTemp", "v": "Log", "line": "SpotLight77 created at Position=<100.0, 200.0, 50.0> Scale=1.5"},
+        {"cat": "LogTemp", "v": "Log", "line": "<Console variable> = <value>"},
+        {"cat": "LogTemp", "v": "Log", "line": "Width=1920, Height=1080, RefreshRate=60"},
+        {"cat": "LogTemp", "v": "Log", "line": 'systemresolution.resx="1920" systemresolution.resy="1080"'},
+        {"cat": "LogTemp", "v": "Fatal", "line": "Assertion failed: IsValid()"},
+        {"cat": "LogTemp", "v": "Log", "line": "Allocated 512KB, freed 2.4MB, peak 1GiB"},
+        {"cat": "LogTemp", "v": "Log", "line": "Allocated 512 KB, freed 2.4 MB, peak 1 GiB"},
+        {"cat": "LogTemp", "v": "Log", "line": "Upload 64Kb at 1.2Mb/s, latency 45ms, timeout 5s, refresh 60Hz"},
+        {"cat": "LogTemp", "v": "Log", "line": "GPU time 145us, CPU time 2.3ms, frame 16.6ms, total 1.5s"},
+        {"cat": "LogTemp", "v": "Log", "line": "Latency: min 50ns, avg 1.2us, p99 4.7ms, offset 320px"},
+        {"cat": "LogTemp", "v": "Log", "line": "Pool: 256M used, 1G total, 512K free, rate 100MB/s"},
+        {"cat": "LogTemp", "v": "Log", "line": "Render target 1920x1080, shadow map 256x256x4, tile 64x64"},
+        {"cat": "None", "v": "Log", "line": 'r.DynamicRes.TestScreenPercentage = "40"'},
+        {"cat": "None", "v": "Log", "line": "r.DynamicRes.TestScreenPercentage= '40'"},
+        {"cat": "None", "v": "Log", "line": "r.DynamicRes.TestScreenPercentage =40"},
+    ]
+    for msg in samples:
+        print(f"  {_format_log_msg(msg)}")
+    print()
+
+    # --- Summary ---
+    print()
+    total = passed + failed
+    if failed:
+        print(f"FAILED: {failed}/{total} tests failed.")
+        sys.exit(1)
+    else:
+        print(f"OK: {total} tests passed.")
diff --git a/Engine/Extras/UnrealRemoteConsole/ue_remote_console.py b/Engine/Extras/UnrealRemoteConsole/ue_remote_console.py
new file mode 100644
--- /dev/null
+++ b/Engine/Extras/UnrealRemoteConsole/ue_remote_console.py
@@ -0,0 +1,358 @@
+#!/usr/bin/env python3
+"""
+UE Remote Console -- cross-platform entry point.
+
+Auto-discovers platform-specific discovery backends by scanning
+Engine/Platforms/*/Extras/UnrealRemoteConsole/discovery.py
+(same convention as ushell).
+
+Usage:
+    python ue_remote_console.py                              # Auto-detect
+    python ue_remote_console.py --platform <name>            # Force platform
+    python ue_remote_console.py --host 10.0.0.5 --port 4600  # Direct TCP
+    python ue_remote_console.py -c "stat unit"               # Single command
+"""
+
+import argparse
+import importlib.util
+import os
+import sys
+import time
+
+_SCRIPT_DIR = os.path.dirname(os.path.abspath(__file__))
+if _SCRIPT_DIR not in sys.path:
+    sys.path.insert(0, _SCRIPT_DIR)
+
+from remote_console import (
+    RemoteConsole, DEFAULT_TIMEOUT, interactive_loop, pipe_mode,
+    _format_log_msg, _disable_colors, _run_self_tests,
+)
+
+
+def _find_engine_platforms_dir():
+    """Find Engine/Platforms/ by walking up from this script."""
+    # Script is at Engine/Extras/UnrealRemoteConsole/ue_remote_console.py
+    d = _SCRIPT_DIR
+    for _ in range(5):
+        candidate = os.path.join(d, 'Engine', 'Platforms')
+        if os.path.isdir(candidate):
+            return candidate
+        # Also check if we're already inside Engine/
+        candidate = os.path.join(d, 'Platforms')
+        if os.path.isdir(candidate) and os.path.basename(d) == 'Engine':
+            return candidate
+        d = os.path.dirname(d)
+    return None
+
+
+def _discover_backends():
+    """Scan Engine/Platforms/*/Extras/UnrealRemoteConsole/discovery.py"""
+    platforms_dir = _find_engine_platforms_dir()
+    if not platforms_dir:
+        return {}
+
+    backends = {}
+    for platform_name in os.listdir(platforms_dir):
+        discovery_path = os.path.join(
+            platforms_dir, platform_name, 'Extras', 'UnrealRemoteConsole', 'discovery.py')
+        if not os.path.isfile(discovery_path):
+            continue
+        try:
+            spec = importlib.util.spec_from_file_location(
+                f"remote_console_discovery_{platform_name}", discovery_path)
+            mod = importlib.util.module_from_spec(spec)
+            spec.loader.exec_module(mod)
+            required = ('PLATFORM_NAME', 'PLATFORM_ALIASES', 'is_available', 'discover', 'resolve')
+            if not all(hasattr(mod, a) for a in required):
+                continue
+            for alias in mod.PLATFORM_ALIASES:
+                backends[alias.lower()] = mod
+            backends[mod.PLATFORM_NAME.lower()] = mod
+        except Exception as e:
+            print(f"Warning: Failed to load {discovery_path}: {e}", file=sys.stderr)
+    return backends
+
+
+def _probe_port(host, port, timeout=1.0):
+    """Try a TCP connect to check if something is listening. Returns True if reachable."""
+    import socket
+    try:
+        sock = socket.socket(socket.AF_INET, socket.SOCK_STREAM)
+        sock.settimeout(timeout)
+        try:
+            sock.connect((host, port))
+            return True
+        except (OSError, socket.timeout):
+            return False
+        finally:
+            try:
+                sock.close()
+            except OSError:
+                pass
+    except Exception:
+        return False
+
+
+def _discover_all_targets(backends, args=None, quiet=False):
+    """Discover targets across all available platform backends.
+
+    Runs SDK discovery sequentially (fast, local), then probes all discovered
+    targets' ports in parallel to verify the UE process is actually running.
+    Returns a list of (platform_mod, target_info) tuples.
+    """
+    from concurrent.futures import ThreadPoolExecutor, as_completed
+
+    # Phase 1: gather candidates from all backends (SDK tool calls, fast)
+    candidates = []  # (mod, target_info, host, port, label)
+    checked_platforms = []
+    for mod in dict.fromkeys(backends.values()):
+        try:
+            if not mod.is_available():
+                continue
+            checked_platforms.append(mod.PLATFORM_NAME)
+            targets = mod.discover(args=args, quiet=True)
+            for t in targets:
+                host, port, label = mod.resolve(t)
+                candidates.append((mod, t, host, port, label))
+            if not targets and not quiet:
+                print(f"  {mod.PLATFORM_NAME}: no targets found", file=sys.stderr)
+        except Exception:
+            continue
+
+    if not candidates:
+        if not quiet and not checked_platforms:
+            print("  No platform SDKs available.", file=sys.stderr)
+        return []
+
+    # Phase 2: probe all candidates in parallel, cancel remaining on first success
+    results = []
+    with ThreadPoolExecutor(max_workers=len(candidates)) as pool:
+        future_to_candidate = {
+            pool.submit(_probe_port, host, port): (mod, t, host, port, label)
+            for mod, t, host, port, label in candidates
+        }
+        try:
+            for future in as_completed(future_to_candidate):
+                mod, t, host, port, label = future_to_candidate[future]
+                try:
+                    if future.result():
+                        results.append((mod, t))
+                    elif not quiet:
+                        print(f"  {label or mod.PLATFORM_NAME}: devkit found but port {port} not responding (game not running?)",
+                              file=sys.stderr)
+                except Exception:
+                    pass
+        finally:
+            # Cancel any probes still in flight (they'll finish on their own
+            # timeout but we won't block waiting for them)
+            for f in future_to_candidate:
+                f.cancel()
+
+    return results
+
+
+def _select_target(all_targets):
+    """Present a selection menu when multiple targets are found.
+
+    Returns the selected (platform_mod, target_info) tuple, or None if aborted.
+    """
+    print("\nMultiple targets found:", file=sys.stderr)
+    for i, (mod, info) in enumerate(all_targets, 1):
+        host, port, label = mod.resolve(info)
+        print(f"  [{i}] {label or mod.PLATFORM_NAME} ({host}:{port})", file=sys.stderr)
+
+    while True:
+        try:
+            choice = input(f"\nSelect target [1-{len(all_targets)}]: ").strip()
+        except (KeyboardInterrupt, EOFError):
+            print("\nAborted.", file=sys.stderr)
+            return None
+        try:
+            idx = int(choice)
+            if 1 <= idx <= len(all_targets):
+                return all_targets[idx - 1]
+        except ValueError:
+            pass
+        print(f"Invalid choice. Enter a number between 1 and {len(all_targets)}.", file=sys.stderr)
+
+
+def main():
+    backends = _discover_backends()
+
+    parser = argparse.ArgumentParser(
+        description="UE Remote Console",
+        formatter_class=argparse.RawDescriptionHelpFormatter,
+        epilog="""
+Examples:
+  %(prog)s                              Auto-detect platform and connect
+  %(prog)s --platform <name>            Force platform backend
+  %(prog)s --host 10.0.0.5 --port 4600  Direct TCP (any platform)
+  %(prog)s -c "stat unit"               Single command
+""")
+    parser.add_argument("--host", help="Target IP (bypasses platform discovery)")
+    parser.add_argument("--port", type=int, help="Target port (bypasses platform discovery)")
+    parser.add_argument("--platform", help="Force platform backend (see --list-platforms)")
+    parser.add_argument("-c", "--command", action="append",
+                        help="Execute command(s) and exit")
+    parser.add_argument("-t", "--timeout", type=float, default=DEFAULT_TIMEOUT,
+                        help=f"Response timeout in seconds (default: {DEFAULT_TIMEOUT})")
+    parser.add_argument("-v", "--verbose", action="store_true",
+                        help="Diagnostic debug output")
+    parser.add_argument("--debug-log-raw", metavar="FILE",
+                        help="Log raw bytes received from server to FILE")
+    parser.add_argument("--debug-log-decoded", metavar="FILE",
+                        help="Log decoded text to FILE (UTF-8)")
+    parser.add_argument("--test", action="store_true",
+                        help="Run self-tests and exit")
+    parser.add_argument("--list-platforms", action="store_true",
+                        help="List discovered platform backends and exit")
+    parser.add_argument("--no-color", action="store_true",
+                        help="Disable colored output")
+
+    for mod in dict.fromkeys(backends.values()):
+        if hasattr(mod, 'add_arguments'):
+            try:
+                mod.add_arguments(parser)
+            except Exception:
+                pass
+
+    args = parser.parse_args()
+
+    if args.no_color:
+        _disable_colors()
+
+    if args.test:
+        _run_self_tests()
+        return
+
+    if args.list_platforms:
+        seen = set()
+        for mod in backends.values():
+            if id(mod) not in seen:
+                seen.add(id(mod))
+                avail = "available" if mod.is_available() else "not available"
+                aliases = ", ".join(mod.PLATFORM_ALIASES)
+                print(f"  {mod.PLATFORM_NAME:<20s} ({avail})  aliases: {aliases}")
+        if not backends:
+            print("  No platform backends found.")
+        return
+
+    # Resolve platform
+    platform_mod = None
+    if args.host and args.port:
+        pass
+    elif args.platform:
+        platform_mod = backends.get(args.platform.lower())
+        if not platform_mod:
+            avail = ", ".join(sorted(set(m.PLATFORM_NAME.lower() for m in backends.values())))
+            print(f"Error: Unknown platform '{args.platform}'. Available: {avail or 'none'}",
+                  file=sys.stderr)
+            sys.exit(1)
+
+    # Discover target
+    host, port, label = args.host, args.port, None
+    if host and not port:
+        print("Error: --port is required with --host.", file=sys.stderr)
+        sys.exit(1)
+    if port and not host:
+        print("Error: --host is required with --port.", file=sys.stderr)
+        sys.exit(1)
+
+    if not host or not port:
+        if platform_mod:
+            # Explicit --platform: discover only from that backend
+            print(f"Using {platform_mod.PLATFORM_NAME} discovery...", file=sys.stderr)
+            targets = platform_mod.discover(args=args)
+            if not targets:
+                print(f"No targets found. Retrying every 1s (Ctrl+C to abort)...", file=sys.stderr)
+                while not targets:
+                    try:
+                        time.sleep(1.0)
+                    except KeyboardInterrupt:
+                        print("\nAborted.", file=sys.stderr)
+                        sys.exit(1)
+                    targets = platform_mod.discover(args=args, quiet=True)
+            target_info = targets[0]
+            host, port, label = platform_mod.resolve(target_info)
+        else:
+            # Auto-detect: discover across ALL platforms, present menu if multiple
+            print("Scanning for targets...", file=sys.stderr)
+            all_targets = _discover_all_targets(backends, args=args)
+            if not all_targets:
+                print("No targets found. Retrying every 1s (Ctrl+C to abort)...", file=sys.stderr)
+                while not all_targets:
+                    try:
+                        time.sleep(1.0)
+                    except KeyboardInterrupt:
+                        print("\nAborted.", file=sys.stderr)
+                        sys.exit(1)
+                    all_targets = _discover_all_targets(backends, args=args, quiet=True)
+
+            if len(all_targets) == 1:
+                platform_mod, target_info = all_targets[0]
+            else:
+                selection = _select_target(all_targets)
+                if selection is None:
+                    sys.exit(1)
+                platform_mod, target_info = selection
+
+            host, port, label = platform_mod.resolve(target_info)
+
+        print(f"Found {label or 'target'} at {host}:{port}", file=sys.stderr)
+
+    console = RemoteConsole(host, port, timeout=args.timeout, verbose=args.verbose,
+                            log_file=args.debug_log_raw, decoded_log_file=args.debug_log_decoded)
+    if label:
+        console.label = label
+
+    if platform_mod and not args.host:
+        console.set_endpoint_discovery(
+            lambda: _rediscover(platform_mod, args))
+
+    print(f"Connecting to {host}:{port}...", file=sys.stderr)
+
+    if not console.connect(quiet=True):
+        print("Connection failed, retrying every 1s (Ctrl+C to abort)...", file=sys.stderr)
+        while not console.connect(quiet=True):
+            try:
+                time.sleep(1.0)
+            except KeyboardInterrupt:
+                print("\nAborted.", file=sys.stderr)
+                sys.exit(1)
+            console.try_rediscover_endpoint()
+
+    print("Connected.", file=sys.stderr)
+
+    try:
+        if args.command:
+            _last_log_time = [time.monotonic()]
+            def _cmd_log_cb(msg):
+                print(f"  {_format_log_msg(msg)}")
+                _last_log_time[0] = time.monotonic()
+            console.set_log_callback(_cmd_log_cb)
+            for cmd in args.command:
+                resp = console.execute(cmd)
+                if resp.get("error"):
+                    print(f"Error: {resp['error']}")
+                _last_log_time[0] = time.monotonic()
+            while time.monotonic() - _last_log_time[0] < 0.15:
+                time.sleep(0.05)
+        elif sys.stdin.isatty():
+            interactive_loop(console)
+        else:
+            pipe_mode(console)
+    except (ConnectionError, TimeoutError) as e:
+        print(f"Connection lost: {e}", file=sys.stderr)
+    finally:
+        console.close(permanent=True)
+
+
+def _rediscover(platform_mod, args):
+    targets = platform_mod.discover(args=args, quiet=True)
+    if targets:
+        return platform_mod.resolve(targets[0])
+    return None, None, None
+
+
+if __name__ == "__main__":
+    main()
